Documentation for build tools

diff --git a/Engine/Documentation/Source/Programming/Development/BuildGraph/ScriptAnatomy/Tasks/PredefinedTasks/BuildGraphScriptPredefinedTasks.INT.udn b/Engine/Documentation/Source/Programming/Development/BuildGraph/ScriptAnatomy/Tasks/PredefinedTasks/BuildGraphScriptPredefinedTasks.INT.udn
--- a/Engine/Documentation/Source/Programming/Development/BuildGraph/ScriptAnatomy/Tasks/PredefinedTasks/BuildGraphScriptPredefinedTasks.INT.udn
+++ b/Engine/Documentation/Source/Programming/Development/BuildGraph/ScriptAnatomy/Tasks/PredefinedTasks/BuildGraphScriptPredefinedTasks.INT.udn
@@ -5,3 +5,618 @@ Description: This is a procedurally generated markdown page.
 version: 5.1
 parent:Programming/Development/BuildGraph/BuildGraphScriptTasks
 
+### AgeStore
+
+Task that strips symbols from a set of files. This task is named after the AGESTORE utility that comes with the Microsoft debugger tools SDK, but is actually a separate implementation. The main difference is that it uses the last modified time rather than last access time to determine which files to delete.
+
+|          |                      |          |                                                                                                                                                                                                                                                                                                                                     |
+| -------- | -------------------- | -------- | ----------------------------------------------------------------------------------------------------------------------------------------------------------------------------------------------------------------------------------------------------------------------------------------------------------------------------------- |
+| Platform | UnrealTargetPlatform | Required | The target platform to age symbols for.                                                                                                                                                                                                                                                                                             |
+| StoreDir | String               | Required | The symbol server directory.                                                                                                                                                                                                                                                                                                        |
+| Days     | Integer              | Required | Number of days worth of symbols to keep.                                                                                                                                                                                                                                                                                            |
+| BuildDir | String               | Optional | The root of the build directory to check for existing buildversion named directories.                                                                                                                                                                                                                                               |
+| Filter   | String               | Optional | A substring to match in directory file names before deleting symbols. This allows the "age store" task to avoid deleting symbols from other builds in the case where multiple builds share the same symbol server. Specific use of the filter value is determined by the symbol server structure defined by the platform toolchain. |
+| If       | Condition            | Optional | Whether to execute this task. It is ignored if this condition evaluates to false.                                                                                                                                                                                                                                                   |
+
+### Aws
+
+Spawns AWS CLI and waits for it to complete.
+
+|                 |           |          |                                                                                   |
+| --------------- | --------- | -------- | --------------------------------------------------------------------------------- |
+| Arguments       | String    | Optional | Arguments for the newly created process.                                          |
+| Environment     | String    | Optional | Environment variables                                                             |
+| EnvironmentFile | String    | Optional | File to read environment from                                                     |
+| LogOutput       | Boolean   | Optional | Write output to the log                                                           |
+| If              | Condition | Optional | Whether to execute this task. It is ignored if this condition evaluates to false. |
+
+### Aws-AssumeRole
+
+Assumes an AWS role.
+
+|                 |           |          |                                                                                   |
+| --------------- | --------- | -------- | --------------------------------------------------------------------------------- |
+| Arn             | String    | Required | Role to assume                                                                    |
+| Session         | String    | Required | Name of this session                                                              |
+| Duration        | Integer   | Optional | Duration of the token in seconds                                                  |
+| Environment     | String    | Optional | Environment variables                                                             |
+| EnvironmentFile | String    | Optional | File to read environment variables from                                           |
+| OutputFile      | String    | Required | Output file for the new environment                                               |
+| If              | Condition | Optional | Whether to execute this task. It is ignored if this condition evaluates to false. |
+
+### Aws-EcsDeploy
+
+Creates a new AWS ECS task definition and updates the ECS service to use this new revision of the task def
+
+|                    |           |          |                                                                                              |
+| ------------------ | --------- | -------- | -------------------------------------------------------------------------------------------- |
+| TaskDefinitionFile | String    | Required | Task definition file to use                                                                  |
+| DockerImage        | String    | Required | Docker image to set in new task definition (will replace %%DOCKER_PATTERN%% with this value) |
+| Version            | String    | Optional | App version to set in new task definition (will replace %%VERSION%% with this value)         |
+| Cluster            | String    | Required | Cluster ARN representing AWS ECS cluster to operate on                                       |
+| Service            | String    | Required | Service name to update and deploy to                                                         |
+| Environment        | String    | Optional | Environment variables                                                                        |
+| EnvironmentFile    | String    | Optional | File to read environment from                                                                |
+| LogOutput          | Boolean   | Optional | Write output to the log                                                                      |
+| If                 | Condition | Optional | Whether to execute this task. It is ignored if this condition evaluates to false.            |
+
+### CallMethod
+
+Invokes a C# method.
+
+|        |           |          |                                                                                   |
+| ------ | --------- | -------- | --------------------------------------------------------------------------------- |
+| Class  | String    | Required | Name of the class to run a method in                                              |
+| Method | String    | Required | The method name to execute.                                                       |
+| Arg1   | String    | Optional | Argument value                                                                    |
+| Arg2   | String    | Optional | Argument value                                                                    |
+| Arg3   | String    | Optional | Argument value                                                                    |
+| Arg4   | String    | Optional | Argument value                                                                    |
+| Arg5   | String    | Optional | Argument value                                                                    |
+| Arg6   | String    | Optional | Argument value                                                                    |
+| Arg7   | String    | Optional | Argument value                                                                    |
+| Arg8   | String    | Optional | Argument value                                                                    |
+| Arg9   | String    | Optional | Argument value                                                                    |
+| Tags   | String    | Optional | List of tags for outputs                                                          |
+| If     | Condition | Optional | Whether to execute this task. It is ignored if this condition evaluates to false. |
+
+### Command
+
+Invokes an AutomationTool child process to run the given command.
+
+|                          |           |          |                                                                                                                                                                   |
+| ------------------------ | --------- | -------- | ----------------------------------------------------------------------------------------------------------------------------------------------------------------- |
+| Name                     | String    | Required | The command name to execute.                                                                                                                                      |
+| Arguments                | String    | Optional | Arguments to be passed to the command.                                                                                                                            |
+| MergeTelemetryWithPrefix | String    | Optional | If non-null, instructs telemetry from the command to be merged into the telemetry for this UAT instance with the given prefix. May be an empty (non-null) string. |
+| If                       | Condition | Optional | Whether to execute this task. It is ignored if this condition evaluates to false.                                                                                 |
+
+### Commandlet
+
+Spawns the editor to run a commandlet.
+
+|            |               |          |                                                                                                             |
+| ---------- | ------------- | -------- | ----------------------------------------------------------------------------------------------------------- |
+| Name       | String        | Required | The commandlet name to execute.                                                                             |
+| Project    | FileSpec      | Optional | The project to run the editor with.                                                                         |
+| Arguments  | String        | Optional | Arguments to be passed to the commandlet.                                                                   |
+| EditorExe  | FileReference | Optional | The editor executable to use. Defaults to the development UnrealEditor executable for the current platform. |
+| ErrorLevel | Integer       | Optional | The minimum exit code, which is treated as an error.                                                        |
+| If         | Condition     | Optional | Whether to execute this task. It is ignored if this condition evaluates to false.                           |
+
+### Compile
+
+Compiles a target with UnrealBuildTool.
+
+|                       |                           |          |                                                                                                                                  |
+| --------------------- | ------------------------- | -------- | -------------------------------------------------------------------------------------------------------------------------------- |
+| Target                | String                    | Optional | The target to compile.                                                                                                           |
+| Configuration         | UnrealTargetConfiguration | Required | The configuration to compile.                                                                                                    |
+| Platform              | UnrealTargetPlatform      | Required | The platform to compile for.                                                                                                     |
+| Project               | FileSpec                  | Optional | The project to compile with.                                                                                                     |
+| Arguments             | String                    | Optional | Additional arguments for UnrealBuildTool.                                                                                        |
+| AllowXGE              | Boolean                   | Optional | Whether to allow using XGE for compilation.                                                                                      |
+| AllowParallelExecutor | Boolean                   | Optional | Whether to allow using the parallel executor for this compile.                                                                   |
+| Clean                 | Boolean                   | Optional | Whether to allow cleaning this target. If unspecified, targets are cleaned if the -Clean argument is passed on the command line. |
+| Tag                   | TagList                   | Optional | Tag to be applied to build products of this task.                                                                                |
+| If                    | Condition                 | Optional | Whether to execute this task. It is ignored if this condition evaluates to false.                                                |
+
+### Cook
+
+Cook a selection of maps for a certain platform
+
+|           |           |          |                                                                                   |
+| --------- | --------- | -------- | --------------------------------------------------------------------------------- |
+| Project   | String    | Required | Project file to be cooked.                                                        |
+| Platform  | String    | Required | The cook platform to target (for example, Windows).                               |
+| Maps      | String    | Optional | List of maps to be cooked, separated by '+' characters.                           |
+| Versioned | Boolean   | Optional | Additional arguments to be passed to the cooker.                                  |
+| Arguments | String    | Optional | Additional arguments to be passed to the cooker.                                  |
+| EditorExe | String    | Optional | Optional path to what editor executable to run for cooking.                       |
+| Tag       | TagList   | Optional | Tag to be applied to build products of this task.                                 |
+| If        | Condition | Optional | Whether to execute this task. It is ignored if this condition evaluates to false. |
+
+### Copy
+
+Copies files from one directory to another.
+
+|                 |           |          |                                                                                   |
+| --------------- | --------- | -------- | --------------------------------------------------------------------------------- |
+| Files           | FileSpec  | Optional | Optional filter to be applied to the list of input files.                         |
+| From            | FileSpec  | Required | The pattern(s) to copy from (for example, Engine/*.txt).                          |
+| To              | FileSpec  | Required | The directory to copy to.                                                         |
+| Overwrite       | Boolean   | Optional | Whether or not to overwrite existing files.                                       |
+| Tag             | TagList   | Optional | Tag to be applied to build products of this task.                                 |
+| ErrorIfNotFound | Boolean   | Optional | Whether or not to throw an error if no files were found to copy                   |
+| If              | Condition | Optional | Whether to execute this task. It is ignored if this condition evaluates to false. |
+
+### CsCompile
+
+Compiles C# project files, and their dependencies.
+
+|                   |           |          |                                                                                                                                                     |
+| ----------------- | --------- | -------- | --------------------------------------------------------------------------------------------------------------------------------------------------- |
+| Project           | String    | Required | The C# project file to compile. Using semicolons, more than one project file can be specified.                                                      |
+| Configuration     | String    | Optional | The configuration to compile.                                                                                                                       |
+| Platform          | String    | Optional | The platform to compile.                                                                                                                            |
+| Target            | String    | Optional | The target to build.                                                                                                                                |
+| Properties        | String    | Optional | Properties for the command                                                                                                                          |
+| Arguments         | String    | Optional | Additional options to pass to the compiler.                                                                                                         |
+| EnumerateOnly     | Boolean   | Optional | Only enumerate build products -- do not actually compile the projects.                                                                              |
+| Tag               | TagList   | Optional | Tag to be applied to build products of this task.                                                                                                   |
+| TagReferences     | TagList   | Optional | Tag to be applied to any non-private references the projects have. (for example, those that are external and not copied into the output directory). |
+| UseSystemCompiler | Boolean   | Optional | Whether to use the system toolchain rather than the bundled UE SDK                                                                                  |
+| If                | Condition | Optional | Whether to execute this task. It is ignored if this condition evaluates to false.                                                                   |
+
+### Delete
+
+Delete a set of files.
+
+|                        |           |          |                                                                                                                     |
+| ---------------------- | --------- | -------- | ------------------------------------------------------------------------------------------------------------------- |
+| Files                  | FileSpec  | Optional | List of file specifications separated by semicolons (for example, *.cpp;Engine/.../*.bat), or the name of a tag set |
+| Directories            | String    | Optional | List of directory names                                                                                             |
+| DeleteEmptyDirectories | Boolean   | Optional | Whether to delete empty directories after deleting the files. Defaults to true.                                     |
+| Verbose                | Boolean   | Optional | Whether or not to use verbose logging.                                                                              |
+| If                     | Condition | Optional | Whether to execute this task. It is ignored if this condition evaluates to false.                                   |
+
+### Docker
+
+Spawns Docker and waits for it to complete.
+
+|                 |           |          |                                                                                   |
+| --------------- | --------- | -------- | --------------------------------------------------------------------------------- |
+| Arguments       | String    | Required | Docker command line arguments                                                     |
+| Environment     | String    | Optional | Environment variables to set                                                      |
+| EnvironmentFile | String    | Optional | File to read environment variables from                                           |
+| WorkingDir      | String    | Optional | Base directory for running the command                                            |
+| If              | Condition | Optional | Whether to execute this task. It is ignored if this condition evaluates to false. |
+
+### Docker-Build
+
+Spawns Docker and waits for it to complete.
+
+|                  |           |          |                                                                                                                |
+| ---------------- | --------- | -------- | -------------------------------------------------------------------------------------------------------------- |
+| BaseDir          | String    | Required | Base directory for the build                                                                                   |
+| Files            | String    | Required | Files to be staged before building the image                                                                   |
+| DockerFile       | String    | Optional | Path to the Dockerfile. Uses the root of basedir if not specified.                                             |
+| DockerIgnoreFile | String    | Optional | Path to a .dockerignore. Will be copied to basedir if specified.                                               |
+| UseBuildKit      | Boolean   | Optional | Use BuildKit in Docker                                                                                         |
+| ProgressOutput   | String    | Optional | Type of progress output (--progress)                                                                           |
+| Tag              | String    | Optional | Tag for the image                                                                                              |
+| Arguments        | String    | Optional | Optional arguments                                                                                             |
+| OverlayDirs      | String    | Optional | List of additional directories to overlay into the staged input files. Allows credentials to be staged, etc... |
+| Environment      | String    | Optional | Environment variables to set                                                                                   |
+| EnvironmentFile  | String    | Optional | File to read environment variables from                                                                        |
+| If               | Condition | Optional | Whether to execute this task. It is ignored if this condition evaluates to false.                              |
+
+### Docker-Compose-Down
+
+Spawns Docker and waits for it to complete.
+
+|           |           |          |                                                                                   |
+| --------- | --------- | -------- | --------------------------------------------------------------------------------- |
+| File      | String    | Required | Path to the docker-compose file                                                   |
+| Arguments | String    | Optional | Arguments for the command                                                         |
+| If        | Condition | Optional | Whether to execute this task. It is ignored if this condition evaluates to false. |
+
+### Docker-Compose-Up
+
+Spawns Docker and waits for it to complete.
+
+|           |           |          |                                                                                   |
+| --------- | --------- | -------- | --------------------------------------------------------------------------------- |
+| File      | String    | Required | Path to the docker-compose file                                                   |
+| Arguments | String    | Optional | Arguments for the command                                                         |
+| If        | Condition | Optional | Whether to execute this task. It is ignored if this condition evaluates to false. |
+
+### Docker-Push
+
+Spawns Docker and waits for it to complete.
+
+|                 |           |          |                                                                                   |
+| --------------- | --------- | -------- | --------------------------------------------------------------------------------- |
+| Repository      | String    | Required | Repository                                                                        |
+| Image           | String    | Required | Source image to push                                                              |
+| TargetImage     | String    | Optional | Name of the target image                                                          |
+| Environment     | String    | Optional | Additional environment variables                                                  |
+| EnvironmentFile | String    | Optional | File to read environment from                                                     |
+| AwsEcr          | Boolean   | Optional | Whether to login to AWS ECR                                                       |
+| If              | Condition | Optional | Whether to execute this task. It is ignored if this condition evaluates to false. |
+
+### DotNet
+
+Spawns Docker and waits for it to complete.
+
+|                 |               |          |                                                                                   |
+| --------------- | ------------- | -------- | --------------------------------------------------------------------------------- |
+| Arguments       | String        | Optional | Docker command line arguments                                                     |
+| BaseDir         | String        | Optional | Base directory for running the command                                            |
+| Environment     | String        | Optional | Environment variables to set                                                      |
+| EnvironmentFile | String        | Optional | File to read environment variables from                                           |
+| ErrorLevel      | Integer       | Optional | The minimum exit code, which is treated as an error.                              |
+| DotNetPath      | FileReference | Optional | Override path to dotnet executable                                                |
+| If              | Condition     | Optional | Whether to execute this task. It is ignored if this condition evaluates to false. |
+
+### FindModifiedFiles
+
+Compiles C# project files, and their dependencies.
+
+|           |               |          |                                                                                                |
+| --------- | ------------- | -------- | ---------------------------------------------------------------------------------------------- |
+| Paths     | String        | Required | The C# project file to compile. Using semicolons, more than one project file can be specified. |
+| Change    | Integer       | Optional | The configuration to compile.                                                                  |
+| MinChange | Integer       | Optional | The configuration to compile.                                                                  |
+| MaxChange | Integer       | Optional | The configuration to compile.                                                                  |
+| Output    | FileReference | Optional | The file to write to                                                                           |
+| If        | Condition     | Optional | Whether to execute this task. It is ignored if this condition evaluates to false.              |
+
+### Git
+
+Spawns Git and waits for it to complete.
+
+|            |           |          |                                                                                   |
+| ---------- | --------- | -------- | --------------------------------------------------------------------------------- |
+| Arguments  | String    | Optional | Git command line arguments                                                        |
+| BaseDir    | String    | Optional | Base directory for running the command                                            |
+| ErrorLevel | Integer   | Optional | The minimum exit code, which is treated as an error.                              |
+| If         | Condition | Optional | Whether to execute this task. It is ignored if this condition evaluates to false. |
+
+### Git-Clone
+
+Clones a Git repository into a local path.
+
+|            |           |          |                                                                                                                |
+| ---------- | --------- | -------- | -------------------------------------------------------------------------------------------------------------- |
+| Dir        | String    | Required | Directory for the repository                                                                                   |
+| Remote     | String    | Optional | The remote to add                                                                                              |
+| Branch     | String    | Required | The branch to check out on the remote                                                                          |
+| ConfigFile | String    | Optional | Configuration file for the repo. This can be used to set up a remote to be fetched and/or provide credentials. |
+| If         | Condition | Optional | Whether to execute this task. It is ignored if this condition evaluates to false.                              |
+
+### Helm
+
+Spawns Helm and waits for it to complete.
+
+|                 |           |          |                                                                                   |
+| --------------- | --------- | -------- | --------------------------------------------------------------------------------- |
+| Chart           | String    | Required | Helm command line arguments                                                       |
+| Deployment      | String    | Required | Name of the release                                                               |
+| Namespace       | String    | Optional | The Kubernetes namespace                                                          |
+| KubeContext     | String    | Optional | The kubectl context                                                               |
+| KubeConfig      | String    | Optional | The kubectl config file to use                                                    |
+| Values          | String    | Optional | Values to set for running the chart                                               |
+| ValuesFile      | String    | Optional | Values to set for running the chart                                               |
+| Environment     | String    | Optional | Environment variables to set                                                      |
+| EnvironmentFile | String    | Optional | File to parse environment variables from                                          |
+| Arguments       | String    | Optional | Additional arguments                                                              |
+| WorkingDir      | String    | Optional | Base directory for running the command                                            |
+| If              | Condition | Optional | Whether to execute this task. It is ignored if this condition evaluates to false. |
+
+### Kubectl
+
+Spawns Kubectl and waits for it to complete.
+
+|           |           |          |                                                                                   |
+| --------- | --------- | -------- | --------------------------------------------------------------------------------- |
+| Arguments | String    | Required | Command line arguments                                                            |
+| BaseDir   | String    | Optional | Base directory for running the command                                            |
+| If        | Condition | Optional | Whether to execute this task. It is ignored if this condition evaluates to false. |
+
+### Log
+
+Print a message (and other optional diagnostic information) to the output log.
+
+|                 |           |          |                                                                                     |
+| --------------- | --------- | -------- | ----------------------------------------------------------------------------------- |
+| Message         | String    | Optional | Message to print out.                                                               |
+| Files           | FileSpec  | Optional | If specified, causes the given list of files to be printed after the given message. |
+| IncludeContents | Boolean   | Optional | If specified, causes the contents of the given files to be printed out.             |
+| If              | Condition | Optional | Whether to execute this task. It is ignored if this condition evaluates to false.   |
+
+### ModifyConfig
+
+Modifies a config file
+
+|         |           |          |                                                                                   |
+| ------- | --------- | -------- | --------------------------------------------------------------------------------- |
+| File    | FileSpec  | Required | Path to the config file                                                           |
+| Section | String    | Required | The section name to modify                                                        |
+| Key     | String    | Required | The property name to set                                                          |
+| Value   | String    | Required | The property value to set                                                         |
+| Tag     | TagList   | Optional | Tag to be applied to the extracted files                                          |
+| If      | Condition | Optional | Whether to execute this task. It is ignored if this condition evaluates to false. |
+
+### ModifyJsonValue
+
+Modifies json files by setting a value specified in the key path
+
+|          |           |          |                                                                                   |
+| -------- | --------- | -------- | --------------------------------------------------------------------------------- |
+| Files    | FileSpec  | Required | json file paths which will be modified                                            |
+| KeyPath  | String    | Required | json key path to find in each file                                                |
+| NewValue | Integer   | Required | new value to apply                                                                |
+| If       | Condition | Optional | Whether to execute this task. It is ignored if this condition evaluates to false. |
+
+### Move
+
+Moves files from one directory to another.
+
+|           |           |          |                                                                                   |
+| --------- | --------- | -------- | --------------------------------------------------------------------------------- |
+| Files     | FileSpec  | Optional | Optional filter to be applied to the list of input files.                         |
+| From      | FileSpec  | Required | The pattern(s) to copy from (for example, Engine/*.txt).                          |
+| To        | FileSpec  | Required | The directory to copy to.                                                         |
+| Overwrite | Boolean   | Optional | Optionally if files should be overwritten, defaults to false.                     |
+| Tag       | TagList   | Optional | Tag to be applied to build products of this task.                                 |
+| If        | Condition | Optional | Whether to execute this task. It is ignored if this condition evaluates to false. |
+
+### MsBuild
+
+Executes MsBuild
+
+|               |           |          |                                                                                                |
+| ------------- | --------- | -------- | ---------------------------------------------------------------------------------------------- |
+| Project       | String    | Required | The C# project file to compile. Using semicolons, more than one project file can be specified. |
+| Configuration | String    | Optional | The configuration to compile.                                                                  |
+| Platform      | String    | Optional | The platform to compile.                                                                       |
+| Arguments     | String    | Optional | Additional options to pass to the compiler.                                                    |
+| Verbosity     | String    | Optional | The MSBuild output verbosity.                                                                  |
+| If            | Condition | Optional | Whether to execute this task. It is ignored if this condition evaluates to false.              |
+
+### PakFile
+
+Creates a PAK file from a given set of files.
+
+|              |               |          |                                                                                                                                   |
+| ------------ | ------------- | -------- | --------------------------------------------------------------------------------------------------------------------------------- |
+| Files        | FileSpec      | Required | List of files, wildcards, and tag sets to add to the pak file, separated by ';' characters.                                       |
+| Output       | FileReference | Required | PAK file to output.                                                                                                               |
+| ResponseFile | FileReference | Optional | Path to a Response File that contains a list of files to add to the pak file -- instead of specifying them individually.          |
+| RebaseDir    | DirectoryList | Optional | Directories to rebase the files relative to. If specified, the shortest path under a listed directory will be used for each file. |
+| Order        | FileReference | Optional | Script that gives the order of files.                                                                                             |
+| Sign         | FileReference | Optional | Encryption keys for this pak file.                                                                                                |
+| Compress     | Boolean       | Optional | Whether to compress files.                                                                                                        |
+| Arguments    | String        | Optional | Additional arguments to pass to UnrealPak.                                                                                        |
+| Tag          | TagList       | Optional | Tag to be applied to build products of this task.                                                                                 |
+| If           | Condition     | Optional | Whether to execute this task. It is ignored if this condition evaluates to false.                                                 |
+
+### Rename
+
+Renames a file, or group of files.
+
+|       |           |          |                                                                                                          |
+| ----- | --------- | -------- | -------------------------------------------------------------------------------------------------------- |
+| Files | FileSpec  | Required | The file or files to rename.                                                                             |
+| From  | String    | Optional | The current file name, or pattern to match (for example, *.txt). Should not include any path separators. |
+| To    | String    | Required | The new name for the file(s). Should not include any path separators.                                    |
+| Tag   | TagList   | Optional | Tag to be applied to the renamed files.                                                                  |
+| If    | Condition | Optional | Whether to execute this task. It is ignored if this condition evaluates to false.                        |
+
+### SanitizeReceipt
+
+Task that tags build products and/or runtime dependencies by reading from *.target files.
+
+|           |                    |          |                                                                                                                                                 |
+| --------- | ------------------ | -------- | ----------------------------------------------------------------------------------------------------------------------------------------------- |
+| Files     | FileSpec           | Required | Set of receipt files (*.target) to read, including wildcards and tag names, separated by semicolons.                                            |
+| EngineDir | DirectoryReference | Optional | Path to the Engine folder, used to expand $(EngineDir) properties in receipt files. Defaults to the Engine directory for the current workspace. |
+| If        | Condition          | Optional | Whether to execute this task. It is ignored if this condition evaluates to false.                                                               |
+
+### SetVersion
+
+Updates the local version files (Engine/Source/Runtime/Launch/Resources/Version.h, Engine/Build/Build.version, and Engine/Source/Programs/Shared/Metadata.cs) with the given version information.
+
+|                  |           |          |                                                                                                                    |
+| ---------------- | --------- | -------- | ------------------------------------------------------------------------------------------------------------------ |
+| Change           | Integer   | Required | The changelist to set in the version files.                                                                        |
+| CompatibleChange | Integer   | Optional | The engine compatible changelist to set in the version files.                                                      |
+| Branch           | String    | Required | The branch string.                                                                                                 |
+| Build            | String    | Optional | The build version string.                                                                                          |
+| Licensee         | Boolean   | Optional | Whether to set the IS_LICENSEE_VERSION flag to true.                                                               |
+| Promoted         | Boolean   | Optional | Whether to set the ENGINE_IS_PROMOTED_BUILD flag to true.                                                          |
+| SkipWrite        | Boolean   | Optional | If set, do not write to the files -- just return the version files that would be updated. Useful for local builds. |
+| Tag              | TagList   | Optional | Tag to be applied to build products of this task.                                                                  |
+| If               | Condition | Optional | Whether to execute this task. It is ignored if this condition evaluates to false.                                  |
+
+### Sign
+
+Signs a set of executable files with an installed certificate.
+
+|       |           |          |                                                                                                                      |
+| ----- | --------- | -------- | -------------------------------------------------------------------------------------------------------------------- |
+| Files | FileSpec  | Required | List of file specifications separated by semicolons (for example, *.cpp;Engine/.../*.bat), or the name of a tag set. |
+| Tag   | TagList   | Optional | Tag to be applied to build products of this task.                                                                    |
+| If    | Condition | Optional | Whether to execute this task. It is ignored if this condition evaluates to false.                                    |
+
+### Spawn
+
+Spawns an external executable and waits for it to complete.
+
+|                 |           |          |                                                                                   |
+| --------------- | --------- | -------- | --------------------------------------------------------------------------------- |
+| Exe             | String    | Required | Executable to spawn.                                                              |
+| Arguments       | String    | Optional | Arguments for the newly created process.                                          |
+| WorkingDir      | String    | Optional | Working directory for spawning the new task                                       |
+| Environment     | String    | Optional | Environment variables to set                                                      |
+| EnvironmentFile | String    | Optional | File to read environment from                                                     |
+| LogOutput       | Boolean   | Optional | Write output to the log                                                           |
+| ErrorLevel      | Integer   | Optional | The minimum exit code, which is treated as an error.                              |
+| If              | Condition | Optional | Whether to execute this task. It is ignored if this condition evaluates to false. |
+
+### Stage
+
+Stages files listed in a build receipt to an output directory.
+
+|               |                           |          |                                                                                   |
+| ------------- | ------------------------- | -------- | --------------------------------------------------------------------------------- |
+| Project       | FileReference             | Optional | The project that this target belongs to.                                          |
+| Target        | String                    | Required | Name of the target to stage.                                                      |
+| Platform      | UnrealTargetPlatform      | Required | Platform to stage.                                                                |
+| Configuration | UnrealTargetConfiguration | Required | Configuration to be staged.                                                       |
+| Architecture  | String                    | Optional | Architecture to be staged.                                                        |
+| ToDir         | DirectoryReference        | Required | Directory that the receipt files should be staged to.                             |
+| Overwrite     | Boolean                   | Optional | Whether to overwrite existing files.                                              |
+| Tag           | TagList                   | Optional | Tag to be applied to build products of this task.                                 |
+| If            | Condition                 | Optional | Whether to execute this task. It is ignored if this condition evaluates to false. |
+
+### Strip
+
+Strips debugging information from a set of files.
+
+|           |                      |          |                                                                                                                |
+| --------- | -------------------- | -------- | -------------------------------------------------------------------------------------------------------------- |
+| Platform  | UnrealTargetPlatform | Required | The platform toolchain to strip binaries.                                                                      |
+| BaseDir   | DirectoryReference   | Optional | The directory to find files in.                                                                                |
+| Files     | FileSpec             | Required | List of file specifications separated by semicolons (for example, Engine/.../*.pdb), or the name of a tag set. |
+| OutputDir | DirectoryReference   | Optional | Output directory for the stripped files. Defaults to the input path, overwriting the input files.              |
+| Tag       | TagList              | Optional | Tag to be applied to build products of this task.                                                              |
+| If        | Condition            | Optional | Whether to execute this task. It is ignored if this condition evaluates to false.                              |
+
+### Submit
+
+Creates a new changelist and submits a set of files to a Perforce stream.
+
+|                 |                    |          |                                                                                                                                                                              |
+| --------------- | ------------------ | -------- | ---------------------------------------------------------------------------------------------------------------------------------------------------------------------------- |
+| Description     | String             | Required | The description for the submitted changelist.                                                                                                                                |
+| Files           | FileSpec           | Required | The files to submit.                                                                                                                                                         |
+| FileType        | String             | Optional | The Perforce file type for the submitted files (for example, binary+FS32).                                                                                                   |
+| Workspace       | String             | Optional | The workspace name. If specified, a new workspace will be created using the given stream and root directory to submit the files. If not, the current workspace will be used. |
+| Stream          | String             | Optional | The stream for the workspace -- defaults to the current stream. Ignored unless the Workspace attribute is also specified.                                                    |
+| Branch          | String             | Optional | Branch for the workspace (legacy P4 depot path). May not be used in conjunction with Stream.                                                                                 |
+| RootDir         | DirectoryReference | Optional | Root directory for the stream. If not specified, defaults to the current root directory.                                                                                     |
+| RevertUnchanged | Boolean            | Optional | Whether to revert unchanged files before attempting to submit.                                                                                                               |
+| Force           | Boolean            | Optional | Force the submit to happen -- even if a resolve is needed (always accept current version).                                                                                   |
+| P4Verbose       | Boolean            | Optional | Allow verbose P4 output (spew).                                                                                                                                              |
+| If              | Condition          | Optional | Whether to execute this task. It is ignored if this condition evaluates to false.                                                                                            |
+
+### SymStore
+
+Task that strips symbols from a set of files.
+
+|              |                      |          |                                                                                                                                               |
+| ------------ | -------------------- | -------- | --------------------------------------------------------------------------------------------------------------------------------------------- |
+| Platform     | UnrealTargetPlatform | Required | The platform toolchain required to handle symbol files.                                                                                       |
+| Files        | String               | Required | List of output files. PDBs will be extracted from this list.                                                                                  |
+| StoreDir     | String               | Required | Output directory for the compressed symbols.                                                                                                  |
+| Product      | String               | Required | Name of the product for the symbol store records.                                                                                             |
+| BuildVersion | String               | Optional | BuildVersion associated with these symbols. Used for clean-up in AgeStore by matching this version against a directory name in a build share. |
+| If           | Condition            | Optional | Whether to execute this task. It is ignored if this condition evaluates to false.                                                             |
+
+### Tag
+
+Applies a tag to a given set of files. The list of files is found by enumerating the tags and file specifications given by the 'Files' parameter. From this list, any files not matched by the 'Filter' parameter are removed, followed by any files matched by the 'Except' parameter.
+
+|           |                    |          |                                                                                                                                                                                                                           |
+| --------- | ------------------ | -------- | ------------------------------------------------------------------------------------------------------------------------------------------------------------------------------------------------------------------------- |
+| BaseDir   | DirectoryReference | Optional | Set the base directory to resolve relative paths and patterns against. If set, any absolute patterns (for example, /Engine/Build/...) are taken to be relative to this path. If not, they are taken to be truly absolute. |
+| Files     | FileSpec           | Optional | Set of files to work from, including wildcards and tag names, separated by semicolons. If set, resolved relative to BaseDir, otherwise resolved to the branch root directory.                                             |
+| FileLists | FileSpec           | Optional | Set of text files to add additional files from. Each file list should have one file per line.                                                                                                                             |
+| Filter    | FileSpec           | Optional | Patterns to filter the list of files by, including tag names or wildcards. If set, may include patterns that apply to the base directory. If not specified, defaults to all files.                                        |
+| Except    | FileSpec           | Optional | Set of patterns to exclude from the matched list. May include tag names of patterns that apply to the base directory.                                                                                                     |
+| With      | TagList            | Required | Name of the tag to apply.                                                                                                                                                                                                 |
+| If        | Condition          | Optional | Whether to execute this task. It is ignored if this condition evaluates to false.                                                                                                                                         |
+
+### TagReceipt
+
+Task that tags build products and/or runtime dependencies by reading from *.target files.
+
+|                     |                    |          |                                                                                                                                                                 |
+| ------------------- | ------------------ | -------- | --------------------------------------------------------------------------------------------------------------------------------------------------------------- |
+| Files               | FileSpec           | Required | Set of receipt files (*.target) to read, including wildcards and tag names, separated by semicolons.                                                            |
+| EngineDir           | DirectoryReference | Optional | Path to the Engine folder, used to expand $(EngineDir) properties in receipt files. Defaults to the Engine directory for the current workspace.                 |
+| ProjectDir          | DirectoryReference | Optional | Path to the project folder, used to expand $(ProjectDir) properties in receipt files. Defaults to the Engine directory for the current workspace -- DEPRECATED. |
+| BuildProducts       | Boolean            | Optional | Whether to tag the Build Products listed in receipts.                                                                                                           |
+| BuildProductType    | String             | Optional | Which type of Build Products to tag (see TargetReceipt.cs - UnrealBuildTool.BuildProductType for valid values).                                                 |
+| RuntimeDependencies | Boolean            | Optional | Whether to tag the Runtime Dependencies listed in receipts.                                                                                                     |
+| StagedFileType      | String             | Optional | Which type of Runtime Dependencies to tag (see TargetReceipt.cs - UnrealBuildTool.StagedFileType for valid values).                                             |
+| With                | TagList            | Required | Name of the tag to apply.                                                                                                                                       |
+| If                  | Condition          | Optional | Whether to execute this task. It is ignored if this condition evaluates to false.                                                                               |
+
+### Unzip
+
+Extract files from a zip archive.
+
+|                |                    |          |                                                                                   |
+| -------------- | ------------------ | -------- | --------------------------------------------------------------------------------- |
+| ZipFile        | FileSpec           | Required | Path to the zip file to extract.                                                  |
+| ToDir          | DirectoryReference | Required | Output directory for the extracted files.                                         |
+| UseLegacyUnzip | Boolean            | Optional | Whether or not to use the legacy unzip code.                                      |
+| Tag            | TagList            | Optional | Tag to be applied to the extracted files.                                         |
+| If             | Condition          | Optional | Whether to execute this task. It is ignored if this condition evaluates to false. |
+
+### Upload
+
+Uploads a set of files to Jupiter for future retrival
+
+|                    |                    |          |                                                                                                                                                                                       |
+| ------------------ | ------------------ | -------- | ------------------------------------------------------------------------------------------------------------------------------------------------------------------------------------- |
+| FromDir            | DirectoryReference | Required | The directory to read compressed files from.                                                                                                                                          |
+| Files              | FileSpec           | Optional | List of file specifications separated by semicolons (for example, *.cpp;Engine/.../*.bat), or the name of a tag set. Relative paths are taken from FromDir.                           |
+| JupiterNamespace   | String             | Required | The jupiter namespace used to upload the build. Used to control who has access to the build.                                                                                          |
+| JupiterKey         | String             | Required | The key of the build as will be used to download the build again. This has to be globally unique for this particular upload.                                                          |
+| ArchiveType        | String             | Required | The type of archive these files are from, will be added to the metadata                                                                                                               |
+| ProjectName        | String             | Required | The name of the project this set of files are associated with, will be added to the metadata                                                                                          |
+| Branch             | String             | Required | The source control branch these files were generated from, will be added to the metadata                                                                                              |
+| Changelist         | String             | Required | The source control revision these files were generated from, will be added to the metadata                                                                                            |
+| JupiterUrl         | String             | Required | Specify the url to the Jupiter instance to upload to                                                                                                                                  |
+| AdditionalMetadata | String             | Optional | Semi-colon separated list of '=' separated key value mappings to add to the metadata. E.g. Foo=bar;spam=eggs                                                                          |
+| LimitMemoryUsage   | Boolean            | Optional | If enabled file content is not kept in memory, results in lower memory usage but increased io as file contents needs to be read multiple times (for hashing as well as during upload) |
+| If                 | Condition          | Optional | Whether to execute this task. It is ignored if this condition evaluates to false.                                                                                                     |
+
+### Wait
+
+Waits a defined number of seconds.
+
+|         |           |          |                                                                                   |
+| ------- | --------- | -------- | --------------------------------------------------------------------------------- |
+| Seconds | Integer   | Required | Number of seconds to wait.                                                        |
+| If      | Condition | Optional | Whether to execute this task. It is ignored if this condition evaluates to false. |
+
+### WriteTextFile
+
+Writes text to a file.
+
+|        |               |          |                                                                                     |
+| ------ | ------------- | -------- | ----------------------------------------------------------------------------------- |
+| File   | FileReference | Required | Path to the file to write.                                                          |
+| Append | Boolean       | Optional | Optional, whether or not to append to the file rather than overwrite.               |
+| Text   | String        | Optional | The text to write to the file.                                                      |
+| Files  | FileSpec      | Optional | If specified, causes the given list of files to be printed after the given message. |
+| Tag    | TagList       | Optional | Tag to be applied to build products of this task.                                   |
+| If     | Condition     | Optional | Whether to execute this task. It is ignored if this condition evaluates to false.   |
+
+### Zip
+
+Compresses files into a zip archive.
+
+|         |                    |          |                                                                                                                                                             |
+| ------- | ------------------ | -------- | ----------------------------------------------------------------------------------------------------------------------------------------------------------- |
+| FromDir | DirectoryReference | Required | The directory to read compressed files from.                                                                                                                |
+| Files   | FileSpec           | Optional | List of file specifications separated by semicolons (for example, *.cpp;Engine/.../*.bat), or the name of a tag set. Relative paths are taken from FromDir. |
+| ZipFile | FileReference      | Required | The zip file to create.                                                                                                                                     |
+| Tag     | TagList            | Optional | Tag to be applied to the created zip file.                                                                                                                  |
+| If      | Condition          | Optional | Whether to execute this task. It is ignored if this condition evaluates to false.                                                                           |
+
