Documentation for build tools

diff --git a/Engine/Documentation/Source/Programming/UnrealBuildSystem/Configuration/BuildConfigProperties/BuildConfigProperties.INT.udn b/Engine/Documentation/Source/Programming/UnrealBuildSystem/Configuration/BuildConfigProperties/BuildConfigProperties.INT.udn
--- a/Engine/Documentation/Source/Programming/UnrealBuildSystem/Configuration/BuildConfigProperties/BuildConfigProperties.INT.udn
+++ b/Engine/Documentation/Source/Programming/UnrealBuildSystem/Configuration/BuildConfigProperties/BuildConfigProperties.INT.udn
@@ -50,6 +50,8 @@ $ bUseUnityBuild : Whether to unify C++ code into larger files for faster compil
 
 $ bForceUnityBuild : Whether to force C++ source files to be combined into larger files for faster compilation.
 
+$ bMergeModuleAndGeneratedUnityFiles : Whether to merge module and generated unity files for faster compilation.
+
 $ bUseAdaptiveUnityBuild : Use a heuristic to determine which files are currently being iterated on and exclude them from unity blobs, result in faster incremental compile times. The current implementation uses the read-only flag to distinguish the working set, assuming that files will be made writable by the source control system if they are being modified. This is true for Perforce, but not for Git.
 
 $ bAdaptiveUnityDisablesOptimizations : Disable optimization for files that are in the adaptive non-unity working set.
diff --git a/Engine/Documentation/Source/Programming/UnrealBuildSystem/TargetFiles/TargetFilesProperties/TargetFilesProperties.INT.udn b/Engine/Documentation/Source/Programming/UnrealBuildSystem/TargetFiles/TargetFilesProperties/TargetFilesProperties.INT.udn
--- a/Engine/Documentation/Source/Programming/UnrealBuildSystem/TargetFiles/TargetFilesProperties/TargetFilesProperties.INT.udn
+++ b/Engine/Documentation/Source/Programming/UnrealBuildSystem/TargetFiles/TargetFilesProperties/TargetFilesProperties.INT.udn
@@ -188,6 +188,8 @@ $ bDisableSymbolCache (Boolean): True if debug symbols that are cached for some
 
 $ bForceUnityBuild (Boolean): Whether to force C++ source files to be combined into larger files for faster compilation.
 
+$ bMergeModuleAndGeneratedUnityFiles (Boolean): Whether to merge module and generated unity files for faster compilation.
+
 $ bUseAdaptiveUnityBuild (Boolean): Use a heuristic to determine which files are currently being iterated on and exclude them from unity blobs, result in faster incremental compile times. The current implementation uses the read-only flag to distinguish the working set, assuming that files will be made writable by the source control system if they are being modified. This is true for Perforce, but not for Git.
 
 $ bAdaptiveUnityDisablesOptimizations (Boolean): Disable optimization for files that are in the adaptive non-unity working set.
