Documentation for build tools

diff --git a/Engine/Documentation/Source/Programming/Development/BuildGraph/ScriptAnatomy/Tasks/PredefinedTasks/BuildGraphScriptPredefinedTasks.INT.udn b/Engine/Documentation/Source/Programming/Development/BuildGraph/ScriptAnatomy/Tasks/PredefinedTasks/BuildGraphScriptPredefinedTasks.INT.udn
--- a/Engine/Documentation/Source/Programming/Development/BuildGraph/ScriptAnatomy/Tasks/PredefinedTasks/BuildGraphScriptPredefinedTasks.INT.udn
+++ b/Engine/Documentation/Source/Programming/Development/BuildGraph/ScriptAnatomy/Tasks/PredefinedTasks/BuildGraphScriptPredefinedTasks.INT.udn
@@ -60,26 +60,6 @@ Creates a new AWS ECS task definition and updates the ECS service to use this ne
 | LogOutput          | Boolean   | Optional | Write output to the log                                                                      |
 | If                 | Condition | Optional | Whether to execute this task. It is ignored if this condition evaluates to false.            |
 
-### CallMethod
-
-Invokes a C# method.
-
-|        |           |          |                                                                                   |
-| ------ | --------- | -------- | --------------------------------------------------------------------------------- |
-| Class  | String    | Required | Name of the class to run a method in                                              |
-| Method | String    | Required | The method name to execute.                                                       |
-| Arg1   | String    | Optional | Argument value                                                                    |
-| Arg2   | String    | Optional | Argument value                                                                    |
-| Arg3   | String    | Optional | Argument value                                                                    |
-| Arg4   | String    | Optional | Argument value                                                                    |
-| Arg5   | String    | Optional | Argument value                                                                    |
-| Arg6   | String    | Optional | Argument value                                                                    |
-| Arg7   | String    | Optional | Argument value                                                                    |
-| Arg8   | String    | Optional | Argument value                                                                    |
-| Arg9   | String    | Optional | Argument value                                                                    |
-| Tags   | String    | Optional | List of tags for outputs                                                          |
-| If     | Condition | Optional | Whether to execute this task. It is ignored if this condition evaluates to false. |
-
 ### Command
 
 Invokes an AutomationTool child process to run the given command.
diff --git a/Engine/Documentation/Source/Programming/UnrealBuildSystem/Configuration/BuildConfigProperties/BuildConfigProperties.INT.udn b/Engine/Documentation/Source/Programming/UnrealBuildSystem/Configuration/BuildConfigProperties/BuildConfigProperties.INT.udn
--- a/Engine/Documentation/Source/Programming/UnrealBuildSystem/Configuration/BuildConfigProperties/BuildConfigProperties.INT.udn
+++ b/Engine/Documentation/Source/Programming/UnrealBuildSystem/Configuration/BuildConfigProperties/BuildConfigProperties.INT.udn
@@ -624,7 +624,7 @@ $ bPrimaryProjectNameFromFolder : If true, sets the primary project name accordi
 
 $ Version : The version of Visual Studio to generate project files for.
 
-$ MaxSharedIncludePaths : Puts the most common include paths in the IncludePath property in the MSBuild project. This significantly reduces Visual Studio memory usage (measured 1.1GB -> 500mb), but seems to be causing issues with Visual Assist. Value here specifies maximum length of the include path list in KB.
+$ MaxSharedIncludePaths : Puts the most common include paths in the IncludePath property in the MSBuild project. This significantly reduces Visual Studio memory usage (measured 1.1GB -> 500mb), but seems to be causing issues with Visual Assist. Value here specifies maximum length of the include path list in KB. Ignored for VS2022.
 
 $ ExcludedIncludePaths : Semi-colon separated list of paths that should not be added to the projects include paths. Useful for omitting third-party headers (e.g ThirdParty/WebRTC) from intellisense suggestions and reducing memory footprints.
 
