Documentation for build tools

diff --git a/Engine/Documentation/Source/Programming/UnrealBuildSystem/Configuration/BuildConfigProperties/BuildConfigProperties.INT.udn b/Engine/Documentation/Source/Programming/UnrealBuildSystem/Configuration/BuildConfigProperties/BuildConfigProperties.INT.udn
--- a/Engine/Documentation/Source/Programming/UnrealBuildSystem/Configuration/BuildConfigProperties/BuildConfigProperties.INT.udn
+++ b/Engine/Documentation/Source/Programming/UnrealBuildSystem/Configuration/BuildConfigProperties/BuildConfigProperties.INT.udn
@@ -624,7 +624,7 @@ $ bPrimaryProjectNameFromFolder : If true, sets the primary project name accordi
 
 $ Version : The version of Visual Studio to generate project files for.
 
-$ MaxSharedIncludePaths : Puts the most common include paths in the IncludePath property in the MSBuild project. This significantly reduces Visual Studio memory usage (measured 1.1GB -> 500mb), but seems to be causing issues with Visual Assist. Value here specifies maximum length of the include path list in KB. Ignored for VS2022.
+$ MaxSharedIncludePaths : Puts the most common include paths in the IncludePath property in the MSBuild project. This significantly reduces Visual Studio memory usage (measured 1.1GB -> 500mb), but seems to be causing issues with Visual Assist. Value here specifies maximum length of the include path list in KB.
 
 $ ExcludedIncludePaths : Semi-colon separated list of paths that should not be added to the projects include paths. Useful for omitting third-party headers (e.g ThirdParty/WebRTC) from intellisense suggestions and reducing memory footprints.
 
