Documentation for build tools

diff --git a/Engine/Documentation/Source/Programming/UnrealBuildSystem/Configuration/BuildConfigProperties/BuildConfigProperties.INT.udn b/Engine/Documentation/Source/Programming/UnrealBuildSystem/Configuration/BuildConfigProperties/BuildConfigProperties.INT.udn
--- a/Engine/Documentation/Source/Programming/UnrealBuildSystem/Configuration/BuildConfigProperties/BuildConfigProperties.INT.udn
+++ b/Engine/Documentation/Source/Programming/UnrealBuildSystem/Configuration/BuildConfigProperties/BuildConfigProperties.INT.udn
@@ -96,6 +96,14 @@ $ bUsePDBFiles : Whether PDB files should be used for Visual C++ builds.
 
 $ bUsePCHFiles : Whether PCH files should be used.
 
+$ bPreprocessDepends : Generate dependency files by preprocessing. This is only recommended when distributing builds as it adds additional overhead.
+
+$ StaticAnalyzer : Whether static code analysis should be enabled.
+
+$ StaticAnalyzerOutputType : The output type to use for the static analyzer.
+
+$ StaticAnalyzerMode : The mode to use for the static analyzer.
+
 $ MinFilesUsingPrecompiledHeader : The minimum number of files that must use a pre-compiled header before it will be created and used.
 
 $ bForcePrecompiledHeaderForGameModules : When enabled, a precompiled header is always generated for game modules, even if there are only a few source files in the module.  This greatly improves compile times for iterative changes on a few files in the project, at the expense of slower full rebuild times for small game projects.  This can be overridden by setting MinFilesUsingPrecompiledHeaderOverride in a module's Build.cs file.
@@ -170,8 +178,6 @@ $ bEnableMemorySanitizer : Enables memory sanitizer (MSan)
 
 $ bTuneDebugInfoForLLDB : Turns on tuning of debug info for LLDB
 
-$ bPreprocessDepends : Generate dependency files by preprocessing. This is only recommended when distributing builds as it adds additional overhead.
-
 $ bUseDSYMFiles : Enables the generation of .dsym files. This can be disabled to enable faster iteration times during development.
 
 $ bWriteSolutionOptionFile : Whether to write a solution option (suo) file for the sln.
diff --git a/Engine/Documentation/Source/Programming/UnrealBuildSystem/TargetFiles/TargetFilesProperties/TargetFilesProperties.INT.udn b/Engine/Documentation/Source/Programming/UnrealBuildSystem/TargetFiles/TargetFilesProperties/TargetFilesProperties.INT.udn
--- a/Engine/Documentation/Source/Programming/UnrealBuildSystem/TargetFiles/TargetFilesProperties/TargetFilesProperties.INT.udn
+++ b/Engine/Documentation/Source/Programming/UnrealBuildSystem/TargetFiles/TargetFilesProperties/TargetFilesProperties.INT.udn
@@ -234,6 +234,14 @@ $ bUsePCHFiles (Boolean): Whether PCH files should be used.
 
 $ bPreprocessOnly (Boolean): Whether to just preprocess source files for this target, and skip compilation
 
+$ bPreprocessDepends (Boolean): Generate dependency files by preprocessing. This is only recommended when distributing builds as it adds additional overhead.
+
+$ StaticAnalyzer (StaticAnalyzer): Whether static code analysis should be enabled.
+
+$ StaticAnalyzerOutputType (StaticAnalyzerOutputType): The output type to use for the static analyzer.
+
+$ StaticAnalyzerMode (StaticAnalyzerMode): The mode to use for the static analyzer.
+
 $ MinFilesUsingPrecompiledHeader (Int32): The minimum number of files that must use a pre-compiled header before it will be created and used.
 
 $ bForcePrecompiledHeaderForGameModules (Boolean): When enabled, a precompiled header is always generated for game modules, even if there are only a few source files in the module.  This greatly improves compile times for iterative changes on a few files in the project, at the expense of slower full rebuild times for small game projects.  This can be overridden by setting MinFilesUsingPrecompiledHeaderOverride in a module's Build.cs file.
