Documentation for build tools

diff --git a/Engine/Documentation/Source/Programming/UnrealBuildSystem/TargetFiles/TargetFilesProperties/TargetFilesProperties.INT.udn b/Engine/Documentation/Source/Programming/UnrealBuildSystem/TargetFiles/TargetFilesProperties/TargetFilesProperties.INT.udn
--- a/Engine/Documentation/Source/Programming/UnrealBuildSystem/TargetFiles/TargetFilesProperties/TargetFilesProperties.INT.udn
+++ b/Engine/Documentation/Source/Programming/UnrealBuildSystem/TargetFiles/TargetFilesProperties/TargetFilesProperties.INT.udn
@@ -112,6 +112,8 @@ $ bCompileIntelMetricsDiscovery (Boolean): Whether to compile IntelMetricsDiscov
 
 $ bCompilePython (Boolean): Whether to compile in python support
 
+$ bUseIris (Boolean): Whether to use Iris.
+
 $ bForceBuildTargetPlatforms (Boolean): Whether to force compiling the target platform modules, even if they wouldn't normally be built.
 
 $ bForceBuildShaderFormats (Boolean): Whether to force compiling shader format modules, even if they wouldn't normally be built.
