Documentation for build tools

diff --git a/Engine/Documentation/Source/Programming/UnrealBuildSystem/Configuration/BuildConfigProperties/BuildConfigProperties.INT.udn b/Engine/Documentation/Source/Programming/UnrealBuildSystem/Configuration/BuildConfigProperties/BuildConfigProperties.INT.udn
--- a/Engine/Documentation/Source/Programming/UnrealBuildSystem/Configuration/BuildConfigProperties/BuildConfigProperties.INT.udn
+++ b/Engine/Documentation/Source/Programming/UnrealBuildSystem/Configuration/BuildConfigProperties/BuildConfigProperties.INT.udn
@@ -100,9 +100,9 @@ $ bPreprocessDepends : Generate dependency files by preprocessing. This is only
 
 $ StaticAnalyzer : Whether static code analysis should be enabled.
 
-$ StaticAnalyzerOutputType : The output type to use for the static analyzer.
+$ StaticAnalyzerOutputType : The output type to use for the static analyzer. This is only supported for Clang.
 
-$ StaticAnalyzerMode : The mode to use for the static analyzer.
+$ StaticAnalyzerMode : The mode to use for the static analyzer. This is only supported for Clang. Shallow mode completes quicker but is generally not recommended.
 
 $ MinFilesUsingPrecompiledHeader : The minimum number of files that must use a pre-compiled header before it will be created and used.
 
diff --git a/Engine/Documentation/Source/Programming/UnrealBuildSystem/ModuleFiles/ModuleFilesProperties/ModuleFilesProperties.INT.udn b/Engine/Documentation/Source/Programming/UnrealBuildSystem/ModuleFiles/ModuleFilesProperties/ModuleFilesProperties.INT.udn
--- a/Engine/Documentation/Source/Programming/UnrealBuildSystem/ModuleFiles/ModuleFilesProperties/ModuleFilesProperties.INT.udn
+++ b/Engine/Documentation/Source/Programming/UnrealBuildSystem/ModuleFiles/ModuleFilesProperties/ModuleFilesProperties.INT.udn
@@ -50,6 +50,8 @@ $ bEnableObjCExceptions (Boolean): Enable objective C exception handling
 
 $ bEnableUndefinedIdentifierWarnings (Boolean): Enable warnings for using undefined identifiers in #if expressions
 
+$ bMergeUnityFiles (Boolean): Whether to merge module and generated unity files for faster compilation.
+
 $ MinSourceFilesForUnityBuildOverride (Int32): The number of source files in this module before unity build will be activated for that module.  If set to anything besides -1, will override the default setting which is controlled by MinGameModuleSourceFilesForUnityBuild
 
 $ MinFilesUsingPrecompiledHeaderOverride (Int32): Overrides BuildConfiguration.MinFilesUsingPrecompiledHeader if non-zero.
diff --git a/Engine/Documentation/Source/Programming/UnrealBuildSystem/TargetFiles/TargetFilesProperties/TargetFilesProperties.INT.udn b/Engine/Documentation/Source/Programming/UnrealBuildSystem/TargetFiles/TargetFilesProperties/TargetFilesProperties.INT.udn
--- a/Engine/Documentation/Source/Programming/UnrealBuildSystem/TargetFiles/TargetFilesProperties/TargetFilesProperties.INT.udn
+++ b/Engine/Documentation/Source/Programming/UnrealBuildSystem/TargetFiles/TargetFilesProperties/TargetFilesProperties.INT.udn
@@ -240,9 +240,13 @@ $ bPreprocessDepends (Boolean): Generate dependency files by preprocessing. This
 
 $ StaticAnalyzer (StaticAnalyzer): Whether static code analysis should be enabled.
 
-$ StaticAnalyzerOutputType (StaticAnalyzerOutputType): The output type to use for the static analyzer.
+$ StaticAnalyzerOutputType (StaticAnalyzerOutputType): The output type to use for the static analyzer. This is only supported for Clang.
 
-$ StaticAnalyzerMode (StaticAnalyzerMode): The mode to use for the static analyzer.
+$ StaticAnalyzerMode (StaticAnalyzerMode): The mode to use for the static analyzer. This is only supported for Clang. Shallow mode completes quicker but is generally not recommended.
+
+$ StaticAnalyzerCheckers (List&lt;String&gt;): The static analyzer checkers that should be enabled rather than the defaults. This is only supported for Clang. See https://clang.llvm.org/docs/analyzer/checkers.html for a full list. Or run: 'clang -Xclang -analyzer-checker-help' or: 'clang -Xclang -analyzer-checker-help-alpha' for the list of experimental checkers.
+
+$ StaticAnalyzerAdditionalCheckers (List&lt;String&gt;): The static analyzer non-default checkers that should be enabled. Unused if StaticAnalyzerCheckers is set. This is only supported for Clang. See https://clang.llvm.org/docs/analyzer/checkers.html for a full list. Or run: 'clang -Xclang -analyzer-checker-help' or: 'clang -Xclang -analyzer-checker-help-alpha' for the list of experimental checkers.
 
 $ MinFilesUsingPrecompiledHeader (Int32): The minimum number of files that must use a pre-compiled header before it will be created and used.
 
@@ -418,6 +422,8 @@ $ DeprecationWarningLevel (WarningLevel): Level to report deprecation warnings a
 
 $ bRetainFramePointers (Boolean): Forces frame pointers to be retained this is usually required when you want reliable callstacks e.g. mallocframeprofiler
 
+$ StaticAnalyzerDisabledCheckers (List&lt;String&gt;): The static analyzer default checkers that should be disabled. Unused if StaticAnalyzerCheckers is set. This is only supported for Clang. See https://clang.llvm.org/docs/analyzer/checkers.html for a full list. Or run: 'clang -Xclang -analyzer-checker-help' or: 'clang -Xclang -analyzer-checker-help-alpha' for the list of experimental checkers.
+
 $ bLegacyPublicIncludePaths (Boolean): Add all the public folders as include paths for the compile environment.
 
 $ LinkType (TargetLinkType): Specifies how to link modules in this target (monolithic or modular). This is currently protected for backwards compatibility. Call the GetLinkType() accessor until support for the deprecated ShouldCompileMonolithic() override has been removed.
