Coding standard fixes: gender-inclusive language

diff --git a/Engine/Documentation/Source/Shared/Tutorials/inEditorGamifiedTutorial/inEditorGamifiedTutorial.INT.udn b/Engine/Documentation/Source/Shared/Tutorials/inEditorGamifiedTutorial/inEditorGamifiedTutorial.INT.udn
--- a/Engine/Documentation/Source/Shared/Tutorials/inEditorGamifiedTutorial/inEditorGamifiedTutorial.INT.udn
+++ b/Engine/Documentation/Source/Shared/Tutorials/inEditorGamifiedTutorial/inEditorGamifiedTutorial.INT.udn
@@ -188,7 +188,7 @@ There's a gap, can't continue
 [/VAR]
 
 
-Due to the gap, he cannot continue forward, but you can help him.
+Due to the gap, the character cannot continue forward, but you can help it.
 1. Click the "Stop" button.
 
 ![](Stage1_08.png)
@@ -204,7 +204,7 @@ Finding our Walkway
 [/VAR]
 
 
-Let us help him! We need to find a Walkway to place. The Content Browser holds all of the assets in the current project.
+Let us help it! We need to find a Walkway to place. The Content Browser holds all of the assets in the current project.
 
 [/EXCERPT:Stage9]
 
@@ -307,7 +307,7 @@ You can also directly input the values needed in the "Details" panel. In the "De
 
 [EXCERPT:Stage14]
 [VAR:StageTitle]
-Now He Will Get the Coin
+Now It Will Get the Coin
 [/VAR]
 
 
@@ -330,7 +330,7 @@ First Coin Collected!
 
 
 
-Success! You have helped our little guy here collect the first coin.
+Success! You have helped our little character here collect the first coin.
 1. Press "Stop" to continue.
 
 ![](Stage1_15.png)
@@ -382,7 +382,7 @@ The Setup
 
 
 
-The "Ramp" in the room is what will be used to help our guy again. Though as you can see, it is not in line with the area your guy is standing, nor is it large enough. You will need to "translate," "rotate," and "scale" it to get it into position.
+The "Ramp" in the room is what will be used to help our character again. Though as you can see, it is not in line with the area your character is standing, nor is it large enough. You will need to "translate," "rotate," and "scale" it to get it into position.
 
 
 
@@ -524,7 +524,7 @@ The final transform of the ramp should be close to:
 Time to Get Coin 2
 [/VAR]
 
-Push play and watch your guy get the second objective!
+Push play and watch your character get the second objective!
 1. Click the "Play" button.
 
 [/EXCERPT:Stage27]
diff --git a/Engine/Source/Developer/Virtualization/Private/VirtualizationSourceControlBackend.cpp b/Engine/Source/Developer/Virtualization/Private/VirtualizationSourceControlBackend.cpp
--- a/Engine/Source/Developer/Virtualization/Private/VirtualizationSourceControlBackend.cpp
+++ b/Engine/Source/Developer/Virtualization/Private/VirtualizationSourceControlBackend.cpp
@@ -28,7 +28,7 @@ namespace UE::Virtualization
 {
 
 /** 
- * A quick and dirty, poor mans implementation of a std::counting_semaphore that we can use to 
+ * A quick and dirty implementation of a std::counting_semaphore that we can use to 
  * limit the number of threads that can create a new perforce connection when pulling or 
  * pushing payloads.
  * 
diff --git a/Engine/Source/Editor/Kismet/Public/Debugging/SKismetDebugTreeView.h b/Engine/Source/Editor/Kismet/Public/Debugging/SKismetDebugTreeView.h
--- a/Engine/Source/Editor/Kismet/Public/Debugging/SKismetDebugTreeView.h
+++ b/Engine/Source/Editor/Kismet/Public/Debugging/SKismetDebugTreeView.h
@@ -147,7 +147,7 @@ protected:
 	void CopyValueToClipboard() const;
 
 protected:
-	// Type of action (poor mans RTTI for the tree, really only used to accelerate Compare checks)
+	// Type of action (a kind of RTTI for the tree, really only used to accelerate Compare checks)
 	EDebugLineType Type;
 
 	static uint16 ActiveTypeBitset;
diff --git a/Engine/Source/Programs/AutomationTool/Scripts/CopyBuildToStagingDirectory.Automation.cs b/Engine/Source/Programs/AutomationTool/Scripts/CopyBuildToStagingDirectory.Automation.cs
--- a/Engine/Source/Programs/AutomationTool/Scripts/CopyBuildToStagingDirectory.Automation.cs
+++ b/Engine/Source/Programs/AutomationTool/Scripts/CopyBuildToStagingDirectory.Automation.cs
@@ -4210,7 +4210,7 @@ namespace AutomationScripts
 
 				if (FilesAdded == 0 && bContinueSearch)
 				{
-					// no files have been deployed at all to this guy, so remove all previously added files and exit the loop as this means we need to deploy everything
+					// no files have been deployed at all to this, so remove all previously added files and exit the loop as this means we need to deploy everything
 					DeployedFiles.Clear();
 					bContinueSearch = false;
 				}
diff --git a/Engine/Source/Programs/AutomationTool/Turnkey/LocalCache.cs b/Engine/Source/Programs/AutomationTool/Turnkey/LocalCache.cs
--- a/Engine/Source/Programs/AutomationTool/Turnkey/LocalCache.cs
+++ b/Engine/Source/Programs/AutomationTool/Turnkey/LocalCache.cs
@@ -220,7 +220,7 @@ namespace Turnkey
 						if (!File.Exists(Filename) || FileTime != SavedTime)
 						{
 							TurnkeyUtils.Log("Cleaning old download cache...");
-							// remove this guy from the cache (delayed after the iterator)
+							// remove this entry from the cache (delayed after the iterator)
 							EntriesToDelete.Add(Pair.Key);
 
 							break;
diff --git a/Engine/Source/Runtime/Core/Private/Misc/Timespan.cpp b/Engine/Source/Runtime/Core/Private/Misc/Timespan.cpp
--- a/Engine/Source/Runtime/Core/Private/Misc/Timespan.cpp
+++ b/Engine/Source/Runtime/Core/Private/Misc/Timespan.cpp
@@ -137,7 +137,7 @@ bool FTimespan::Parse(const FString& TimespanString, FTimespan& OutTimespan)
 		Tokens.AddDefaulted();
 	}
 
-	// poor man's token verification
+	// token verification
 	for (const FString& Token : Tokens)
 	{
 		if (!Token.IsEmpty() && !Token.IsNumeric())
diff --git a/Engine/Source/Runtime/CoreUObject/Private/UObject/Obj.cpp b/Engine/Source/Runtime/CoreUObject/Private/UObject/Obj.cpp
--- a/Engine/Source/Runtime/CoreUObject/Private/UObject/Obj.cpp
+++ b/Engine/Source/Runtime/CoreUObject/Private/UObject/Obj.cpp
@@ -4050,7 +4050,7 @@ bool StaticExec( UWorld* InWorld, const TCHAR* Cmd, FOutputDevice& Ar )
 					TotalCnt++;
 				}
 			}
-			// poor mans sort
+			// sort
 			for (int32 CurrentNum = MaxNum; CurrentNum > 1; CurrentNum--)
 			{
 				for (int32 Index = 0; Index < IndexSet.Components.Num(); Index++)
diff --git a/Engine/Source/Runtime/Engine/Classes/Components/ActorComponent.h b/Engine/Source/Runtime/Engine/Classes/Components/ActorComponent.h
--- a/Engine/Source/Runtime/Engine/Classes/Components/ActorComponent.h
+++ b/Engine/Source/Runtime/Engine/Classes/Components/ActorComponent.h
@@ -694,7 +694,7 @@ public:
 	 * ReadyForReplication gets called on replicated components when their owning actor is officially ready for replication.
 	 * Called after InitializeComponent but before BeginPlay. From there it will only be set false when the component is destroyed.
 	 * This is where you want to register your replicated subobjects if you already possess some. 
-	 * A component can get replicated before HasBegunPlay() is true if inside a tick or in BeginPlay() an RPC is called on him.
+	 * A component can get replicated before HasBegunPlay() is true if inside a tick or in BeginPlay() an RPC is called on it.
 	 * Requires component to be registered, initialized and set to replicate.
 	 */
 	virtual void ReadyForReplication();
diff --git a/Engine/Source/Runtime/ImageCore/Private/ImageCore.cpp b/Engine/Source/Runtime/ImageCore/Private/ImageCore.cpp
--- a/Engine/Source/Runtime/ImageCore/Private/ImageCore.cpp
+++ b/Engine/Source/Runtime/ImageCore/Private/ImageCore.cpp
@@ -816,12 +816,12 @@ void FImage::Swap(FImage & Other)
 {
 	::Swap(RawData,Other.RawData);
 	// FImageInfo can be moved with assignment
-	FImageInfo * Me = this;
-	FImageInfo * Him = &Other;
+	FImageInfo * ThisImage = this;
+	FImageInfo * OtherImage = &Other;
 	FImageInfo Temp;
-	Temp = *Me;
-	*Me = *Him;
-	*Him = Temp;
+	Temp = *ThisImage;
+	*ThisImage = *OtherImage;
+	*OtherImage = Temp;
 }
 
 void FImage::ChangeFormat(ERawImageFormat::Type DestFormat, EGammaSpace DestGammaSpace)
