Documentation for build tools

diff --git a/Engine/Documentation/Source/Programming/UnrealBuildSystem/Configuration/BuildConfigProperties/BuildConfigProperties.INT.udn b/Engine/Documentation/Source/Programming/UnrealBuildSystem/Configuration/BuildConfigProperties/BuildConfigProperties.INT.udn
--- a/Engine/Documentation/Source/Programming/UnrealBuildSystem/Configuration/BuildConfigProperties/BuildConfigProperties.INT.udn
+++ b/Engine/Documentation/Source/Programming/UnrealBuildSystem/Configuration/BuildConfigProperties/BuildConfigProperties.INT.udn
@@ -156,6 +156,8 @@ $ bPublicSymbolsByDefault : Whether to expose all symbols as public by default o
 
 $ CppStandard : Which C++ stanard to use for compiling this target
 
+$ CStandard : Which C standard to use for compiling this target
+
 $ bStopSNDBSCompilationAfterErrors : When enabled, SN-DBS will stop compiling targets after a compile error occurs.  Recommended, as it saves computing resources for others.
 
 $ bXGENoWatchdogThread : Whether to use the no_watchdog_thread option to prevent VS2015 toolchain stalls.
diff --git a/Engine/Documentation/Source/Programming/UnrealBuildSystem/ModuleFiles/ModuleFilesProperties/ModuleFilesProperties.INT.udn b/Engine/Documentation/Source/Programming/UnrealBuildSystem/ModuleFiles/ModuleFilesProperties/ModuleFilesProperties.INT.udn
--- a/Engine/Documentation/Source/Programming/UnrealBuildSystem/ModuleFiles/ModuleFilesProperties/ModuleFilesProperties.INT.udn
+++ b/Engine/Documentation/Source/Programming/UnrealBuildSystem/ModuleFiles/ModuleFilesProperties/ModuleFilesProperties.INT.udn
@@ -146,6 +146,8 @@ $ bValidateCircularDependencies (Boolean): Whether circular dependencies will be
 
 $ CppStandard (CppStandardVersion): Which stanard to use for compiling this module
 
+$ CStandard (CStandardVersion): Which standard to use for compiling this module
+
 $ ModuleSymbolVisibility (SymbolVisibility): Control visibility of symbols
 
 $ Name (String): Name of this module
diff --git a/Engine/Documentation/Source/Programming/UnrealBuildSystem/TargetFiles/TargetFilesProperties/TargetFilesProperties.INT.udn b/Engine/Documentation/Source/Programming/UnrealBuildSystem/TargetFiles/TargetFilesProperties/TargetFilesProperties.INT.udn
--- a/Engine/Documentation/Source/Programming/UnrealBuildSystem/TargetFiles/TargetFilesProperties/TargetFilesProperties.INT.udn
+++ b/Engine/Documentation/Source/Programming/UnrealBuildSystem/TargetFiles/TargetFilesProperties/TargetFilesProperties.INT.udn
@@ -324,6 +324,8 @@ $ bAllowNonUFSIniWhenCooked (Boolean): Whether to load non-ufs ini files in cook
 
 $ CppStandard (CppStandardVersion): Which C++ stanard to use for compiling this target
 
+$ CStandard (CStandardVersion): Which C standard to use for compiling this target
+
 $ BuildVersion (String): The build version string
 
 $ GlobalDefinitions (List&lt;String&gt;): Macros to define globally across the whole target.
