Documentation for build tools

diff --git a/Engine/Documentation/Source/Programming/UnrealBuildSystem/Configuration/BuildConfigProperties/BuildConfigProperties.INT.udn b/Engine/Documentation/Source/Programming/UnrealBuildSystem/Configuration/BuildConfigProperties/BuildConfigProperties.INT.udn
--- a/Engine/Documentation/Source/Programming/UnrealBuildSystem/Configuration/BuildConfigProperties/BuildConfigProperties.INT.udn
+++ b/Engine/Documentation/Source/Programming/UnrealBuildSystem/Configuration/BuildConfigProperties/BuildConfigProperties.INT.udn
@@ -86,6 +86,8 @@ $ NumIncludedBytesPerUnityCPP : An approximate number of bytes of C++ code to ta
 
 $ bStressTestUnity : Whether to stress test the C++ unity build robustness by including all C++ files files in a project from a single unified file.
 
+$ bDetailedUnityFiles : Whether to add additional information to the unity files, such as '_of_X' in the file name.
+
 $ bDisableDebugInfo : Whether to globally disable debug info generation; see DebugInfoHeuristics.cs for per-config and per-platform options.
 
 $ bDisableDebugInfoForGeneratedCode : Whether to disable debug info generation for generated files. This improves link times for modules that have a lot of generated glue code.
@@ -96,6 +98,12 @@ $ bUsePDBFiles : Whether PDB files should be used for Visual C++ builds.
 
 $ bUsePCHFiles : Whether PCH files should be used.
 
+$ bDeterministic : Set flags require for determinstic linking (experimental, may not be fully supported). Deterministic compiling is controlled via ModuleRules.
+
+$ bForceDeterministic : Force set flags require for determinstic compiling and linking (experimental, may not be fully supported). This setting is only recommended for testing, instead:
+        * Set bDeterministic on a per module basis in ModuleRules to control deterministic compiling.
+        * Set bDeterministic on a per target basis in TargetRules to control deterministic linking.
+
 $ bPreprocessDepends : Generate dependency files by preprocessing. This is only recommended when distributing builds as it adds additional overhead.
 
 $ StaticAnalyzer : Whether static code analysis should be enabled.
@@ -232,8 +240,6 @@ $ bCompilerTrace : Outputs compile timing information so that it can be analyzed
 
 $ bShowIncludes : Print out files that are included by each source file
 
-$ bDeterministic : Set flags require for determinstic compiles (experimental)
-
 ### ModuleConfiguration
 
 $ DisableUnityBuild : List of modules to disable unity builds for
diff --git a/Engine/Documentation/Source/Programming/UnrealBuildSystem/ModuleFiles/ModuleFilesProperties/ModuleFilesProperties.INT.udn b/Engine/Documentation/Source/Programming/UnrealBuildSystem/ModuleFiles/ModuleFilesProperties/ModuleFilesProperties.INT.udn
--- a/Engine/Documentation/Source/Programming/UnrealBuildSystem/ModuleFiles/ModuleFilesProperties/ModuleFilesProperties.INT.udn
+++ b/Engine/Documentation/Source/Programming/UnrealBuildSystem/ModuleFiles/ModuleFilesProperties/ModuleFilesProperties.INT.udn
@@ -38,6 +38,8 @@ $ ShortName (String): Specifies an alternate name for intermediate directories a
 
 $ bTreatAsEngineModule (Boolean): Whether this module should be treated as an engine module (eg. using engine definitions, PCHs, compiled with optimizations enabled in DebugGame configurations, etc...). Initialized to a default based on the rules assembly it was created from.
 
+$ bDeterministic (Boolean): Set flags for determinstic compiles (experimental, may not be fully supported). Deterministic linking is controlled via TargetRules.
+
 $ bUseRTTI (Boolean): Use run time type information
 
 $ bUseAVX (Boolean): Direct the compiler to generate AVX instructions wherever SSE or AVX intrinsics are used, on the platforms that support it. Note that by enabling this you are changing the minspec for the PC platform, and the resultant executable will crash on machines without AVX support.
diff --git a/Engine/Documentation/Source/Programming/UnrealBuildSystem/TargetFiles/TargetFilesProperties/TargetFilesProperties.INT.udn b/Engine/Documentation/Source/Programming/UnrealBuildSystem/TargetFiles/TargetFilesProperties/TargetFilesProperties.INT.udn
--- a/Engine/Documentation/Source/Programming/UnrealBuildSystem/TargetFiles/TargetFilesProperties/TargetFilesProperties.INT.udn
+++ b/Engine/Documentation/Source/Programming/UnrealBuildSystem/TargetFiles/TargetFilesProperties/TargetFilesProperties.INT.udn
@@ -164,6 +164,8 @@ $ bCompileFreeType (Boolean): True if we need FreeType support.
 
 $ bCompileForSize (Boolean): True if we want to favor optimizing size over speed.
 
+$ OptimizationLevel (OptimizationMode): Allows to fine tune optimizations level for speed and\or code size
+
 $ bForceCompileDevelopmentAutomationTests (Boolean): Whether to compile development automation tests.
 
 $ bForceCompilePerformanceAutomationTests (Boolean): Whether to compile performance automation tests.
@@ -222,6 +224,8 @@ $ NumIncludedBytesPerUnityCPP (Int32): An approximate number of bytes of C++ cod
 
 $ bStressTestUnity (Boolean): Whether to stress test the C++ unity build robustness by including all C++ files files in a project from a single unified file.
 
+$ bDetailedUnityFiles (Boolean): Whether to add additional information to the unity files, such as '_of_X' in the file name.
+
 $ bForceDebugInfo (Boolean): Whether to force debug info to be generated.
 
 $ bDisableDebugInfo (Boolean): Whether to globally disable debug info generation; see DebugInfoHeuristics.cs for per-config and per-platform options.
@@ -234,6 +238,12 @@ $ bUsePDBFiles (Boolean): Whether PDB files should be used for Visual C++ builds
 
 $ bUsePCHFiles (Boolean): Whether PCH files should be used.
 
+$ bDeterministic (Boolean): Set flags require for determinstic linking (experimental, may not be fully supported). Deterministic compiling is controlled via ModuleRules.
+
+$ bForceDeterministic (Boolean): Force set flags require for determinstic compiling and linking (experimental, may not be fully supported). This setting is only recommended for testing, instead:
+        * Set bDeterministic on a per module basis in ModuleRules to control deterministic compiling.
+        * Set bDeterministic on a per target basis in TargetRules to control deterministic linking.
+
 $ bPreprocessOnly (Boolean): Whether to just preprocess source files for this target, and skip compilation
 
 $ bPreprocessDepends (Boolean): Generate dependency files by preprocessing. This is only recommended when distributing builds as it adds additional overhead.
