Documentation for build tools

diff --git a/Engine/Documentation/Source/Programming/UnrealBuildSystem/Configuration/BuildConfigProperties/BuildConfigProperties.INT.udn b/Engine/Documentation/Source/Programming/UnrealBuildSystem/Configuration/BuildConfigProperties/BuildConfigProperties.INT.udn
--- a/Engine/Documentation/Source/Programming/UnrealBuildSystem/Configuration/BuildConfigProperties/BuildConfigProperties.INT.udn
+++ b/Engine/Documentation/Source/Programming/UnrealBuildSystem/Configuration/BuildConfigProperties/BuildConfigProperties.INT.udn
@@ -192,6 +192,8 @@ $ bUseDSYMFiles : Enables the generation of .dsym files. This can be disabled to
 
 $ bWriteSolutionOptionFile : Whether to write a solution option (suo) file for the sln.
 
+$ bVsConfigFile : Whether to write a .vsconfig file next to the sln to suggest components to install.
+
 $ bAddFastPDBToProjects : Whether to add the -FastPDB option to build command lines by default.
 
 $ bUsePerFileIntellisense : Whether to generate per-file intellisense data.
