Documentation for build tools

diff --git a/Engine/Documentation/Source/Programming/UnrealBuildSystem/Configuration/BuildConfigProperties/BuildConfigProperties.INT.udn b/Engine/Documentation/Source/Programming/UnrealBuildSystem/Configuration/BuildConfigProperties/BuildConfigProperties.INT.udn
--- a/Engine/Documentation/Source/Programming/UnrealBuildSystem/Configuration/BuildConfigProperties/BuildConfigProperties.INT.udn
+++ b/Engine/Documentation/Source/Programming/UnrealBuildSystem/Configuration/BuildConfigProperties/BuildConfigProperties.INT.udn
@@ -246,6 +246,10 @@ $ bCompilerTrace : Outputs compile timing information so that it can be analyzed
 
 $ bShowIncludes : Print out files that are included by each source file
 
+### TargetRules
+
+$ bDisableInliningGenCpps : Disable supports for inlining gen.cpps
+
 ### ModuleConfiguration
 
 $ DisableUnityBuild : List of modules to disable unity builds for
@@ -254,6 +258,10 @@ $ EnableOptimizeCode : List of modules to enable optimizations for
 
 $ DisableOptimizeCode : List of modules to disable optimizations for
 
+### TestTargetRules
+
+$ bDisableInliningGenCpps : Disable supports for inlining gen.cpps
+
 ### FASTBuild
 
 $ FBuildExecutablePath : Used to specify the location of fbuild.exe if the distributed binary isn't being used
diff --git a/Engine/Documentation/Source/Programming/UnrealBuildSystem/TargetFiles/TargetFilesProperties/TargetFilesProperties.INT.udn b/Engine/Documentation/Source/Programming/UnrealBuildSystem/TargetFiles/TargetFilesProperties/TargetFilesProperties.INT.udn
--- a/Engine/Documentation/Source/Programming/UnrealBuildSystem/TargetFiles/TargetFilesProperties/TargetFilesProperties.INT.udn
+++ b/Engine/Documentation/Source/Programming/UnrealBuildSystem/TargetFiles/TargetFilesProperties/TargetFilesProperties.INT.udn
@@ -320,6 +320,8 @@ $ bParseTimingInfoForTracing (Boolean): Whether to parse timing data into a trac
 
 $ bPublicSymbolsByDefault (Boolean): Whether to expose all symbols as public by default on POSIX platforms
 
+$ bDisableInliningGenCpps (Boolean): Disable supports for inlining gen.cpps
+
 $ ToolChainName (String): Allows overriding the toolchain to be created for this target. This must match the name of a class declared in the UnrealBuildTool assembly.
 
 $ bDisableUnverifiedCertificates (Boolean): Whether to allow engine configuration to determine if we can load unverified certificates.
