Documentation for build tools

diff --git a/Engine/Documentation/Source/Programming/UnrealBuildSystem/Configuration/BuildConfigProperties/BuildConfigProperties.INT.udn b/Engine/Documentation/Source/Programming/UnrealBuildSystem/Configuration/BuildConfigProperties/BuildConfigProperties.INT.udn
--- a/Engine/Documentation/Source/Programming/UnrealBuildSystem/Configuration/BuildConfigProperties/BuildConfigProperties.INT.udn
+++ b/Engine/Documentation/Source/Programming/UnrealBuildSystem/Configuration/BuildConfigProperties/BuildConfigProperties.INT.udn
@@ -180,6 +180,8 @@ $ BaseLogFileName : Specifies the file to use for logging.
 
 $ bStripSymbols : Whether to strip iOS symbols or not (implied by Shipping config).
 
+$ bSkipClangValidation : Disables clang build verification checks on static libraries
+
 $ bEnableAddressSanitizer : Enables address sanitizer (ASan). Only supported for Visual Studio 2019 16.7.0 and up.
 
 $ bEnableThreadSanitizer : Enables thread sanitizer (TSan).
@@ -228,6 +230,8 @@ $ Compiler : Version of the compiler toolchain to use on Windows platform. A val
 
 $ CompilerVersion : The specific toolchain version to use. This may be a specific version number (for example, "14.13.26128"), the string "Latest" to select the newest available version, or the string "Preview" to select the newest available preview version. By default, and if it is available, we use the toolchain version indicated by WindowsPlatform.DefaultToolChainVersion (otherwise, we use the latest version).
 
+$ bAllowClangLinker : True if we should use the Clang linker (LLD) when we are compiling with Clang, or Intel linker (xilink\xilib) when we are compiling with Intel oneAPI, otherwise we use the MSVC linker.
+
 $ WindowsSdkVersion : The specific Windows SDK version to use. This may be a specific version number (for example, "8.1", "10.0" or "10.0.10150.0"), or the string "Latest", to select the newest available version. By default, and if it is available, we use the Windows SDK version indicated by WindowsPlatform.DefaultWindowsSdkVersion (otherwise, we use the latest version).
 
 $ StaticAnalyzer : The static analyzer to use.
diff --git a/Engine/Documentation/Source/Programming/UnrealBuildSystem/ModuleFiles/ModuleFilesProperties/ModuleFilesProperties.INT.udn b/Engine/Documentation/Source/Programming/UnrealBuildSystem/ModuleFiles/ModuleFilesProperties/ModuleFilesProperties.INT.udn
--- a/Engine/Documentation/Source/Programming/UnrealBuildSystem/ModuleFiles/ModuleFilesProperties/ModuleFilesProperties.INT.udn
+++ b/Engine/Documentation/Source/Programming/UnrealBuildSystem/ModuleFiles/ModuleFilesProperties/ModuleFilesProperties.INT.udn
@@ -64,6 +64,8 @@ $ MinSourceFilesForUnityBuildOverride (Int32): The number of source files in thi
 
 $ MinFilesUsingPrecompiledHeaderOverride (Int32): Overrides BuildConfiguration.MinFilesUsingPrecompiledHeader if non-zero.
 
+$ NumIncludedBytesPerUnityCPPOverride (Int32): Overrides Target.NumIncludedBytesPerUnityCPP if non-zero.
+
 $ bBuildLocallyWithSNDBS (Boolean): Module uses a #import so must be built locally when compiling with SN-DBS
 
 $ IsRedistributableOverride (Nullable&lt;Boolean&gt;): Redistribution override flag for this module.
diff --git a/Engine/Documentation/Source/Programming/UnrealBuildSystem/TargetFiles/TargetFilesProperties/TargetFilesProperties.INT.udn b/Engine/Documentation/Source/Programming/UnrealBuildSystem/TargetFiles/TargetFilesProperties/TargetFilesProperties.INT.udn
--- a/Engine/Documentation/Source/Programming/UnrealBuildSystem/TargetFiles/TargetFilesProperties/TargetFilesProperties.INT.udn
+++ b/Engine/Documentation/Source/Programming/UnrealBuildSystem/TargetFiles/TargetFilesProperties/TargetFilesProperties.INT.udn
@@ -276,6 +276,8 @@ $ bEnableCppModules (Boolean): Whether to enable support for C++20 modules
 
 $ bEnableCppCoroutinesForEvaluation (Boolean): Whether to enable support for C++20 coroutines This option is provided to facilitate evaluation of the feature. Expect the name of the option to change. Use of coroutines in with UE is untested and unsupported.
 
+$ bEnableProcessPriorityControl (Boolean): Whether to enable engine's ability to set process priority on runtime. This option requires some environment setup on Linux, that's why it's disabled by default. Project has to opt-in this feature as it has to guarantee correct setup.
+
 $ bUseMallocProfiler (Boolean): If true, then enable memory profiling in the build (defines USE_MALLOC_PROFILER=1 and forces bOmitFramePointers=false).
 
 $ bUseSharedPCHs (Boolean): Enables "Shared PCHs", a feature which significantly speeds up compile times by attempting to share certain PCH files between modules that UBT detects is including those PCH's header files.
