Documentation for build tools

diff --git a/Engine/Documentation/Source/Programming/UnrealBuildSystem/Configuration/BuildConfigProperties/BuildConfigProperties.INT.udn b/Engine/Documentation/Source/Programming/UnrealBuildSystem/Configuration/BuildConfigProperties/BuildConfigProperties.INT.udn
--- a/Engine/Documentation/Source/Programming/UnrealBuildSystem/Configuration/BuildConfigProperties/BuildConfigProperties.INT.udn
+++ b/Engine/Documentation/Source/Programming/UnrealBuildSystem/Configuration/BuildConfigProperties/BuildConfigProperties.INT.udn
@@ -124,6 +124,10 @@ $ bAllowLTCG : Whether to allow the use of link time code generation (LTCG).
 
 $ bPreferThinLTO : When Link Time Code Generation (LTCG) is enabled, whether to prefer using the lighter weight version on supported platforms.
 
+$ ThinLTOCacheDirectory : Directory where to put the ThinLTO cache on supported platforms.
+
+$ ThinLTOCachePruningArguments : Arguments that will be applied to prune the ThinLTO cache on supported platforms. Arguments will only be applied if ThinLTOCacheDirectory is set.
+
 $ bPGOProfile : Whether to enable Profile Guided Optimization (PGO) instrumentation in this build.
 
 $ bPGOOptimize : Whether to optimize this build with Profile Guided Optimization (PGO).
diff --git a/Engine/Documentation/Source/Programming/UnrealBuildSystem/TargetFiles/TargetFilesProperties/TargetFilesProperties.INT.udn b/Engine/Documentation/Source/Programming/UnrealBuildSystem/TargetFiles/TargetFilesProperties/TargetFilesProperties.INT.udn
--- a/Engine/Documentation/Source/Programming/UnrealBuildSystem/TargetFiles/TargetFilesProperties/TargetFilesProperties.INT.udn
+++ b/Engine/Documentation/Source/Programming/UnrealBuildSystem/TargetFiles/TargetFilesProperties/TargetFilesProperties.INT.udn
@@ -264,6 +264,10 @@ $ bAllowLTCG (Boolean): Whether to allow the use of link time code generation (L
 
 $ bPreferThinLTO (Boolean): When Link Time Code Generation (LTCG) is enabled, whether to prefer using the lighter weight version on supported platforms.
 
+$ ThinLTOCacheDirectory (String): Directory where to put the ThinLTO cache on supported platforms.
+
+$ ThinLTOCachePruningArguments (String): Arguments that will be applied to prune the ThinLTO cache on supported platforms. Arguments will only be applied if ThinLTOCacheDirectory is set.
+
 $ bPGOProfile (Boolean): Whether to enable Profile Guided Optimization (PGO) instrumentation in this build.
 
 $ bPGOOptimize (Boolean): Whether to optimize this build with Profile Guided Optimization (PGO).
