Documentation for build tools

diff --git a/Engine/Documentation/Source/Programming/Development/BuildGraph/ScriptAnatomy/Tasks/PredefinedTasks/BuildGraphScriptPredefinedTasks.INT.udn b/Engine/Documentation/Source/Programming/Development/BuildGraph/ScriptAnatomy/Tasks/PredefinedTasks/BuildGraphScriptPredefinedTasks.INT.udn
--- a/Engine/Documentation/Source/Programming/Development/BuildGraph/ScriptAnatomy/Tasks/PredefinedTasks/BuildGraphScriptPredefinedTasks.INT.udn
+++ b/Engine/Documentation/Source/Programming/Development/BuildGraph/ScriptAnatomy/Tasks/PredefinedTasks/BuildGraphScriptPredefinedTasks.INT.udn
@@ -97,6 +97,7 @@ Compiles a target with UnrealBuildTool.
 | Arguments             | String                    | Optional | Additional arguments for UnrealBuildTool.                                                                                        |
 | AllowXGE              | Boolean                   | Optional | Whether to allow using XGE for compilation.                                                                                      |
 | AllowParallelExecutor | Boolean                   | Optional | Whether to allow using the parallel executor for this compile.                                                                   |
+| AllowAllCores         | Boolean                   | Optional | Whether to allow UBT to use all available cores, when AllowXGE is disabled.                                                      |
 | Clean                 | Boolean                   | Optional | Whether to allow cleaning this target. If unspecified, targets are cleaned if the -Clean argument is passed on the command line. |
 | Tag                   | TagList                   | Optional | Tag to be applied to build products of this task.                                                                                |
 | If                    | Condition                 | Optional | Whether to execute this task. It is ignored if this condition evaluates to false.                                                |
