Documentation for build tools

diff --git a/Engine/Documentation/Source/Programming/UnrealBuildSystem/ModuleFiles/ModuleFilesProperties/ModuleFilesProperties.INT.udn b/Engine/Documentation/Source/Programming/UnrealBuildSystem/ModuleFiles/ModuleFilesProperties/ModuleFilesProperties.INT.udn
--- a/Engine/Documentation/Source/Programming/UnrealBuildSystem/ModuleFiles/ModuleFilesProperties/ModuleFilesProperties.INT.udn
+++ b/Engine/Documentation/Source/Programming/UnrealBuildSystem/ModuleFiles/ModuleFilesProperties/ModuleFilesProperties.INT.udn
@@ -52,6 +52,8 @@ $ bEnableObjCExceptions (Boolean): Enable objective C exception handling
 
 $ bEnableUndefinedIdentifierWarnings (Boolean): Enable warnings for using undefined identifiers in #if expressions
 
+$ bDisableStaticAnalysis (Boolean): Disable all static analysis - clang, msvc, pvs-studio.
+
 $ StaticAnalyzerCheckers (HashSet&lt;String&gt;): The static analyzer checkers that should be enabled rather than the defaults. This is only supported for Clang. See https://clang.llvm.org/docs/analyzer/checkers.html for a full list. Or run: 'clang -Xclang -analyzer-checker-help' or: 'clang -Xclang -analyzer-checker-help-alpha' for the list of experimental checkers.
 
 $ StaticAnalyzerDisabledCheckers (HashSet&lt;String&gt;): The static analyzer default checkers that should be disabled. Unused if StaticAnalyzerCheckers is populated. This is only supported for Clang. This overrides the default disabled checkers, which are deadcode.DeadStores and security.FloatLoopCounter See https://clang.llvm.org/docs/analyzer/checkers.html for a full list. Or run: 'clang -Xclang -analyzer-checker-help' or: 'clang -Xclang -analyzer-checker-help-alpha' for the list of experimental checkers.
