Documentation for build tools

diff --git a/Engine/Documentation/Source/Programming/Development/BuildGraph/ScriptAnatomy/Tasks/PredefinedTasks/BuildGraphScriptPredefinedTasks.INT.udn b/Engine/Documentation/Source/Programming/Development/BuildGraph/ScriptAnatomy/Tasks/PredefinedTasks/BuildGraphScriptPredefinedTasks.INT.udn
--- a/Engine/Documentation/Source/Programming/Development/BuildGraph/ScriptAnatomy/Tasks/PredefinedTasks/BuildGraphScriptPredefinedTasks.INT.udn
+++ b/Engine/Documentation/Source/Programming/Development/BuildGraph/ScriptAnatomy/Tasks/PredefinedTasks/BuildGraphScriptPredefinedTasks.INT.udn
@@ -216,15 +216,16 @@ Spawns Docker and waits for it to complete.
 
 Spawns Docker and waits for it to complete.
 
-|                 |           |          |                                                                                   |
-| --------------- | --------- | -------- | --------------------------------------------------------------------------------- |
-| Repository      | String    | Required | Repository                                                                        |
-| Image           | String    | Required | Source image to push                                                              |
-| TargetImage     | String    | Optional | Name of the target image                                                          |
-| Environment     | String    | Optional | Additional environment variables                                                  |
-| EnvironmentFile | String    | Optional | File to read environment from                                                     |
-| AwsEcr          | Boolean   | Optional | Whether to login to AWS ECR                                                       |
-| If              | Condition | Optional | Whether to execute this task. It is ignored if this condition evaluates to false. |
+|                    |           |          |                                                                                   |
+| ------------------ | --------- | -------- | --------------------------------------------------------------------------------- |
+| Repository         | String    | Required | Repository                                                                        |
+| Image              | String    | Required | Source image to push                                                              |
+| TargetImage        | String    | Optional | Name of the target image                                                          |
+| Environment        | String    | Optional | Additional environment variables                                                  |
+| EnvironmentFile    | String    | Optional | File to read environment from                                                     |
+| AwsEcr             | Boolean   | Optional | Whether to login to AWS ECR                                                       |
+| RepositoryAuthFile | String    | Optional | Path to a json file for authentication to the repository for pushing.             |
+| If                 | Condition | Optional | Whether to execute this task. It is ignored if this condition evaluates to false. |
 
 ### DotNet
 
