Documentation for build tools

diff --git a/Engine/Documentation/Source/Programming/UnrealBuildSystem/Configuration/BuildConfigProperties/BuildConfigProperties.INT.udn b/Engine/Documentation/Source/Programming/UnrealBuildSystem/Configuration/BuildConfigProperties/BuildConfigProperties.INT.udn
--- a/Engine/Documentation/Source/Programming/UnrealBuildSystem/Configuration/BuildConfigProperties/BuildConfigProperties.INT.udn
+++ b/Engine/Documentation/Source/Programming/UnrealBuildSystem/Configuration/BuildConfigProperties/BuildConfigProperties.INT.udn
@@ -138,6 +138,8 @@ $ bOmitFramePointers : Whether to omit frame pointers or not. Disabling is usefu
 
 $ bUseMallocProfiler : If true, then enable memory profiling in the build (defines USE_MALLOC_PROFILER=1 and forces bOmitFramePointers=false).
 
+$ bShaderCompilerWorkerTrace : If true, then enable Unreal Insights (utrace) profiling in the build for the Shader Compiler Worker (defines USE_SHADER_COMPILER_WORKER_TRACE=1).
+
 $ bUseSharedPCHs : Enables "Shared PCHs", a feature which significantly speeds up compile times by attempting to share certain PCH files between modules that UBT detects is including those PCH's header files.
 
 $ bUseShippingPhysXLibraries : True if Development and Release builds should use the release configuration of PhysX/APEX.
diff --git a/Engine/Documentation/Source/Programming/UnrealBuildSystem/TargetFiles/TargetFilesProperties/TargetFilesProperties.INT.udn b/Engine/Documentation/Source/Programming/UnrealBuildSystem/TargetFiles/TargetFilesProperties/TargetFilesProperties.INT.udn
--- a/Engine/Documentation/Source/Programming/UnrealBuildSystem/TargetFiles/TargetFilesProperties/TargetFilesProperties.INT.udn
+++ b/Engine/Documentation/Source/Programming/UnrealBuildSystem/TargetFiles/TargetFilesProperties/TargetFilesProperties.INT.udn
@@ -284,6 +284,8 @@ $ bEnableProcessPriorityControl (Boolean): Whether to enable engine's ability to
 
 $ bUseMallocProfiler (Boolean): If true, then enable memory profiling in the build (defines USE_MALLOC_PROFILER=1 and forces bOmitFramePointers=false).
 
+$ bShaderCompilerWorkerTrace (Boolean): If true, then enable Unreal Insights (utrace) profiling in the build for the Shader Compiler Worker (defines USE_SHADER_COMPILER_WORKER_TRACE=1).
+
 $ bUseSharedPCHs (Boolean): Enables "Shared PCHs", a feature which significantly speeds up compile times by attempting to share certain PCH files between modules that UBT detects is including those PCH's header files.
 
 $ bUseShippingPhysXLibraries (Boolean): True if Development and Release builds should use the release configuration of PhysX/APEX.
