Documentation for build tools

diff --git a/Engine/Documentation/Source/Programming/UnrealBuildSystem/Configuration/BuildConfigProperties/BuildConfigProperties.INT.udn b/Engine/Documentation/Source/Programming/UnrealBuildSystem/Configuration/BuildConfigProperties/BuildConfigProperties.INT.udn
--- a/Engine/Documentation/Source/Programming/UnrealBuildSystem/Configuration/BuildConfigProperties/BuildConfigProperties.INT.udn
+++ b/Engine/Documentation/Source/Programming/UnrealBuildSystem/Configuration/BuildConfigProperties/BuildConfigProperties.INT.udn
@@ -106,6 +106,8 @@ $ bForceDeterministic : Force set flags require for determinstic compiling and l
         * Set bDeterministic on a per module basis in ModuleRules to control deterministic compiling.
         * Set bDeterministic on a per target basis in TargetRules to control deterministic linking.
 
+$ bForceIncludePCHHeadersForGenCppFilesWhenPCHIsDisabled : Whether PCH headers should be force included for gen.cpp files when PCH is disabled.
+
 $ bPreprocessDepends : Generate dependency files by preprocessing. This is only recommended when distributing builds as it adds additional overhead.
 
 $ StaticAnalyzer : Whether static code analysis should be enabled.
diff --git a/Engine/Documentation/Source/Programming/UnrealBuildSystem/ModuleFiles/ModuleFilesProperties/ModuleFilesProperties.INT.udn b/Engine/Documentation/Source/Programming/UnrealBuildSystem/ModuleFiles/ModuleFilesProperties/ModuleFilesProperties.INT.udn
--- a/Engine/Documentation/Source/Programming/UnrealBuildSystem/ModuleFiles/ModuleFilesProperties/ModuleFilesProperties.INT.udn
+++ b/Engine/Documentation/Source/Programming/UnrealBuildSystem/ModuleFiles/ModuleFilesProperties/ModuleFilesProperties.INT.udn
@@ -70,6 +70,8 @@ $ NumIncludedBytesPerUnityCPPOverride (Int32): Overrides Target.NumIncludedBytes
 
 $ bBuildLocallyWithSNDBS (Boolean): Module uses a #import so must be built locally when compiling with SN-DBS
 
+$ bEnableNonInlinedGenCppWarnings (Boolean): Enable warnings for when there are .gen.cpp files that could be inlined in a matching handwritten cpp file
+
 $ IsRedistributableOverride (Nullable&lt;Boolean&gt;): Redistribution override flag for this module.
 
 $ bLegalToDistributeObjectCode (Boolean): Whether the output from this module can be publicly distributed, even if it has code/ dependencies on modules that are not (i.e. CarefullyRedist, NotForLicensees, NoRedist). This should be used when you plan to release binaries but not source.
diff --git a/Engine/Documentation/Source/Programming/UnrealBuildSystem/TargetFiles/TargetFilesProperties/TargetFilesProperties.INT.udn b/Engine/Documentation/Source/Programming/UnrealBuildSystem/TargetFiles/TargetFilesProperties/TargetFilesProperties.INT.udn
--- a/Engine/Documentation/Source/Programming/UnrealBuildSystem/TargetFiles/TargetFilesProperties/TargetFilesProperties.INT.udn
+++ b/Engine/Documentation/Source/Programming/UnrealBuildSystem/TargetFiles/TargetFilesProperties/TargetFilesProperties.INT.udn
@@ -244,6 +244,8 @@ $ bForceDeterministic (Boolean): Force set flags require for determinstic compil
         * Set bDeterministic on a per module basis in ModuleRules to control deterministic compiling.
         * Set bDeterministic on a per target basis in TargetRules to control deterministic linking.
 
+$ bForceIncludePCHHeadersForGenCppFilesWhenPCHIsDisabled (Boolean): Whether PCH headers should be force included for gen.cpp files when PCH is disabled.
+
 $ bPreprocessOnly (Boolean): Whether to just preprocess source files for this target, and skip compilation
 
 $ bPreprocessDepends (Boolean): Generate dependency files by preprocessing. This is only recommended when distributing builds as it adds additional overhead.
