Documentation for build tools

diff --git a/Engine/Documentation/Source/Programming/UnrealBuildSystem/Configuration/BuildConfigProperties/BuildConfigProperties.INT.udn b/Engine/Documentation/Source/Programming/UnrealBuildSystem/Configuration/BuildConfigProperties/BuildConfigProperties.INT.udn
--- a/Engine/Documentation/Source/Programming/UnrealBuildSystem/Configuration/BuildConfigProperties/BuildConfigProperties.INT.udn
+++ b/Engine/Documentation/Source/Programming/UnrealBuildSystem/Configuration/BuildConfigProperties/BuildConfigProperties.INT.udn
@@ -246,7 +246,15 @@ $ StaticAnalyzer : The static analyzer to use.
 
 $ StaticAnalyzerOutputType : The output type to use for the static analyzer.
 
-$ bStrictConformanceMode : Enables strict standard conformance mode (/permissive-) in VS2017+.
+$ bStrictConformanceMode : Enables strict standard conformance mode (/permissive-).
+
+$ bUpdatedCPPMacro : Enables updated __cplusplus macro (/Zc:__cplusplus).
+
+$ bStrictInlineConformance : Enables inline conformance (Remove unreferenced COMDAT) (/Zc:inline).
+
+$ bStrictPreprocessorConformance : Enables new preprocessor conformance (/Zc:preprocessor).
+
+$ bStrictEnumTypesConformance : Enables enum types conformance (/Zc:enumTypes) in VS2022 17.4 Preview 4.0+.
 
 $ PCHMemoryAllocationFactor : Determines the amount of memory that the compiler allocates to construct precompiled headers (/Zm).
 
