Documentation for build tools

diff --git a/Engine/Documentation/Source/Programming/Development/BuildGraph/ScriptAnatomy/Tasks/PredefinedTasks/BuildGraphScriptPredefinedTasks.INT.udn b/Engine/Documentation/Source/Programming/Development/BuildGraph/ScriptAnatomy/Tasks/PredefinedTasks/BuildGraphScriptPredefinedTasks.INT.udn
--- a/Engine/Documentation/Source/Programming/Development/BuildGraph/ScriptAnatomy/Tasks/PredefinedTasks/BuildGraphScriptPredefinedTasks.INT.udn
+++ b/Engine/Documentation/Source/Programming/Development/BuildGraph/ScriptAnatomy/Tasks/PredefinedTasks/BuildGraphScriptPredefinedTasks.INT.udn
@@ -384,6 +384,19 @@ Creates a PAK file from a given set of files.
 | Tag          | TagList       | Optional | Tag to be applied to build products of this task.                                                                                 |
 | If           | Condition     | Optional | Whether to execute this task. It is ignored if this condition evaluates to false.                                                 |
 
+### RandomData
+
+Creates files containing random data in the specified output directory. Used for generating test data for the temp storage system.
+
+|           |           |          |                                                                                   |
+| --------- | --------- | -------- | --------------------------------------------------------------------------------- |
+| Size      | Integer   | Optional | The size of each file.                                                            |
+| Count     | Integer   | Optional | Number of files to write.                                                         |
+| Different | Boolean   | Optional | Whether to generate different data for each output file.                          |
+| OutputDir | String    | Optional | Output directory                                                                  |
+| Tag       | String    | Optional | Optional filter to be applied to the list of input files.                         |
+| If        | Condition | Optional | Whether to execute this task. It is ignored if this condition evaluates to false. |
+
 ### Rename
 
 Renames a file, or group of files.
diff --git a/Engine/Documentation/Source/Programming/UnrealBuildSystem/Configuration/BuildConfigProperties/BuildConfigProperties.INT.udn b/Engine/Documentation/Source/Programming/UnrealBuildSystem/Configuration/BuildConfigProperties/BuildConfigProperties.INT.udn
--- a/Engine/Documentation/Source/Programming/UnrealBuildSystem/Configuration/BuildConfigProperties/BuildConfigProperties.INT.udn
+++ b/Engine/Documentation/Source/Programming/UnrealBuildSystem/Configuration/BuildConfigProperties/BuildConfigProperties.INT.udn
@@ -438,6 +438,8 @@ $ IncludeCppSource : Whether to include C++ targets
 
 $ bIncludeDotNetPrograms : True if we should include csharp program projects in the generated solution. Pass "-DotNet" to enable this.
 
+$ bIncludeTempTargets : Whether to include temporary targets generated by UAT to support content only projects with non-default settings.
+
 $ PrimaryProjectName : Name of the primary project file -- for example, the base file name for the Visual Studio solution file, or the Xcode project file on Mac.
 
 $ bPrimaryProjectNameFromFolder : If true, sets the primary project name according to the name of the folder it is in.
@@ -462,6 +464,8 @@ $ IncludeCppSource : Whether to include C++ targets
 
 $ bIncludeDotNetPrograms : True if we should include csharp program projects in the generated solution. Pass "-DotNet" to enable this.
 
+$ bIncludeTempTargets : Whether to include temporary targets generated by UAT to support content only projects with non-default settings.
+
 $ PrimaryProjectName : Name of the primary project file -- for example, the base file name for the Visual Studio solution file, or the Xcode project file on Mac.
 
 $ bPrimaryProjectNameFromFolder : If true, sets the primary project name according to the name of the folder it is in.
@@ -486,6 +490,8 @@ $ IncludeCppSource : Whether to include C++ targets
 
 $ bIncludeDotNetPrograms : True if we should include csharp program projects in the generated solution. Pass "-DotNet" to enable this.
 
+$ bIncludeTempTargets : Whether to include temporary targets generated by UAT to support content only projects with non-default settings.
+
 $ PrimaryProjectName : Name of the primary project file -- for example, the base file name for the Visual Studio solution file, or the Xcode project file on Mac.
 
 $ bPrimaryProjectNameFromFolder : If true, sets the primary project name according to the name of the folder it is in.
@@ -510,6 +516,8 @@ $ IncludeCppSource : Whether to include C++ targets
 
 $ bIncludeDotNetPrograms : True if we should include csharp program projects in the generated solution. Pass "-DotNet" to enable this.
 
+$ bIncludeTempTargets : Whether to include temporary targets generated by UAT to support content only projects with non-default settings.
+
 $ PrimaryProjectName : Name of the primary project file -- for example, the base file name for the Visual Studio solution file, or the Xcode project file on Mac.
 
 $ bPrimaryProjectNameFromFolder : If true, sets the primary project name according to the name of the folder it is in.
@@ -534,6 +542,8 @@ $ IncludeCppSource : Whether to include C++ targets
 
 $ bIncludeDotNetPrograms : True if we should include csharp program projects in the generated solution. Pass "-DotNet" to enable this.
 
+$ bIncludeTempTargets : Whether to include temporary targets generated by UAT to support content only projects with non-default settings.
+
 $ PrimaryProjectName : Name of the primary project file -- for example, the base file name for the Visual Studio solution file, or the Xcode project file on Mac.
 
 $ bPrimaryProjectNameFromFolder : If true, sets the primary project name according to the name of the folder it is in.
@@ -558,6 +568,8 @@ $ IncludeCppSource : Whether to include C++ targets
 
 $ bIncludeDotNetPrograms : True if we should include csharp program projects in the generated solution. Pass "-DotNet" to enable this.
 
+$ bIncludeTempTargets : Whether to include temporary targets generated by UAT to support content only projects with non-default settings.
+
 $ PrimaryProjectName : Name of the primary project file -- for example, the base file name for the Visual Studio solution file, or the Xcode project file on Mac.
 
 $ bPrimaryProjectNameFromFolder : If true, sets the primary project name according to the name of the folder it is in.
@@ -582,6 +594,8 @@ $ IncludeCppSource : Whether to include C++ targets
 
 $ bIncludeDotNetPrograms : True if we should include csharp program projects in the generated solution. Pass "-DotNet" to enable this.
 
+$ bIncludeTempTargets : Whether to include temporary targets generated by UAT to support content only projects with non-default settings.
+
 $ PrimaryProjectName : Name of the primary project file -- for example, the base file name for the Visual Studio solution file, or the Xcode project file on Mac.
 
 $ bPrimaryProjectNameFromFolder : If true, sets the primary project name according to the name of the folder it is in.
@@ -606,6 +620,8 @@ $ IncludeCppSource : Whether to include C++ targets
 
 $ bIncludeDotNetPrograms : True if we should include csharp program projects in the generated solution. Pass "-DotNet" to enable this.
 
+$ bIncludeTempTargets : Whether to include temporary targets generated by UAT to support content only projects with non-default settings.
+
 $ PrimaryProjectName : Name of the primary project file -- for example, the base file name for the Visual Studio solution file, or the Xcode project file on Mac.
 
 $ bPrimaryProjectNameFromFolder : If true, sets the primary project name according to the name of the folder it is in.
@@ -618,6 +634,8 @@ $ AddDebugAttachConfig : Whether VS Code project generation should include debug
 
 $ AddDebugCoreConfig : Whether VS Code project generation should include debug configurations to allow core dump debugging
 
+$ NoCompileCommands : Do not create compile commands json files with compiler arguments for each file; works better with VS Code extension using UBT server mode.
+
 $ bIncludeDocumentation : True if we should include documentation in the generated projects.
 
 $ bUsePrecompiled : True if build targets should pass the -useprecompiled argument.
@@ -636,6 +654,8 @@ $ IncludeCppSource : Whether to include C++ targets
 
 $ bIncludeDotNetPrograms : True if we should include csharp program projects in the generated solution. Pass "-DotNet" to enable this.
 
+$ bIncludeTempTargets : Whether to include temporary targets generated by UAT to support content only projects with non-default settings.
+
 $ PrimaryProjectName : Name of the primary project file -- for example, the base file name for the Visual Studio solution file, or the Xcode project file on Mac.
 
 $ bPrimaryProjectNameFromFolder : If true, sets the primary project name according to the name of the folder it is in.
@@ -660,6 +680,8 @@ $ IncludeCppSource : Whether to include C++ targets
 
 $ bIncludeDotNetPrograms : True if we should include csharp program projects in the generated solution. Pass "-DotNet" to enable this.
 
+$ bIncludeTempTargets : Whether to include temporary targets generated by UAT to support content only projects with non-default settings.
+
 $ PrimaryProjectName : Name of the primary project file -- for example, the base file name for the Visual Studio solution file, or the Xcode project file on Mac.
 
 $ bPrimaryProjectNameFromFolder : If true, sets the primary project name according to the name of the folder it is in.
@@ -696,6 +718,8 @@ $ IncludeCppSource : Whether to include C++ targets
 
 $ bIncludeDotNetPrograms : True if we should include csharp program projects in the generated solution. Pass "-DotNet" to enable this.
 
+$ bIncludeTempTargets : Whether to include temporary targets generated by UAT to support content only projects with non-default settings.
+
 $ PrimaryProjectName : Name of the primary project file -- for example, the base file name for the Visual Studio solution file, or the Xcode project file on Mac.
 
 $ bPrimaryProjectNameFromFolder : If true, sets the primary project name according to the name of the folder it is in.
@@ -720,6 +744,8 @@ $ IncludeCppSource : Whether to include C++ targets
 
 $ bIncludeDotNetPrograms : True if we should include csharp program projects in the generated solution. Pass "-DotNet" to enable this.
 
+$ bIncludeTempTargets : Whether to include temporary targets generated by UAT to support content only projects with non-default settings.
+
 $ PrimaryProjectName : Name of the primary project file -- for example, the base file name for the Visual Studio solution file, or the Xcode project file on Mac.
 
 $ bPrimaryProjectNameFromFolder : If true, sets the primary project name according to the name of the folder it is in.
diff --git a/Engine/Documentation/Source/Programming/UnrealBuildSystem/ModuleFiles/ModuleFilesProperties/ModuleFilesProperties.INT.udn b/Engine/Documentation/Source/Programming/UnrealBuildSystem/ModuleFiles/ModuleFilesProperties/ModuleFilesProperties.INT.udn
--- a/Engine/Documentation/Source/Programming/UnrealBuildSystem/ModuleFiles/ModuleFilesProperties/ModuleFilesProperties.INT.udn
+++ b/Engine/Documentation/Source/Programming/UnrealBuildSystem/ModuleFiles/ModuleFilesProperties/ModuleFilesProperties.INT.udn
@@ -184,4 +184,6 @@ $ bUseUnity (Boolean): If unity builds are enabled this can be used to override
 
 $ bLegacyPublicIncludePaths (Boolean): Whether this module qualifies included headers from other modules relative to the root of their 'Public' folder. This reduces the number of search paths that have to be passed to the compiler, improving performance and reducing the length of the compiler command line.
 
+$ bLegacyParentIncludePaths (Boolean): Whether this module qualifies included headers from other modules relative to the parent directory. This reduces the number of search paths that have to be passed to the compiler, improving performance and reducing the length of the compiler command line.
+
 
diff --git a/Engine/Documentation/Source/Programming/UnrealBuildSystem/TargetFiles/TargetFilesProperties/TargetFilesProperties.INT.udn b/Engine/Documentation/Source/Programming/UnrealBuildSystem/TargetFiles/TargetFilesProperties/TargetFilesProperties.INT.udn
--- a/Engine/Documentation/Source/Programming/UnrealBuildSystem/TargetFiles/TargetFilesProperties/TargetFilesProperties.INT.udn
+++ b/Engine/Documentation/Source/Programming/UnrealBuildSystem/TargetFiles/TargetFilesProperties/TargetFilesProperties.INT.udn
@@ -420,8 +420,6 @@ $ bWithServerCode (Boolean): Compile server-only code.
 
 $ bWithServerCodeTrusted (Boolean): Compile trusted-server-only code.
 
-$ bWithServerCodeUntrusted (Boolean): Compile untrusted-server-only code.
-
 $ bFNameOutlineNumber (Boolean): Compile with FName storing the number part in the name table. Saves memory when most names are not numbered and those that are are referenced multiple times. The game and engine must ensure they reuse numbered names similarly to name strings to avoid leaking memory.
 
 $ bWithPushModel (Boolean): When enabled, Push Model Networking support will be compiled in. This can help reduce CPU overhead of networking, at the cost of more memory. Always enabled in editor builds.
@@ -448,6 +446,8 @@ $ bRetainFramePointers (Boolean): Forces frame pointers to be retained this is u
 
 $ bLegacyPublicIncludePaths (Boolean): Add all the public folders as include paths for the compile environment.
 
+$ bLegacyParentIncludePaths (Boolean): Add all the parent folders as include paths for the compile environment.
+
 $ LinkType (TargetLinkType): Specifies how to link modules in this target (monolithic or modular). This is currently protected for backwards compatibility. Call the GetLinkType() accessor until support for the deprecated ShouldCompileMonolithic() override has been removed.
 
 $ LaunchModuleName (String): Specifies the name of the launch module. For modular builds, this is the module that is compiled into the target's executable.
