Documentation for build tools

diff --git a/Engine/Documentation/Source/Programming/Development/BuildGraph/ScriptAnatomy/Tasks/PredefinedTasks/BuildGraphScriptPredefinedTasks.INT.udn b/Engine/Documentation/Source/Programming/Development/BuildGraph/ScriptAnatomy/Tasks/PredefinedTasks/BuildGraphScriptPredefinedTasks.INT.udn
--- a/Engine/Documentation/Source/Programming/Development/BuildGraph/ScriptAnatomy/Tasks/PredefinedTasks/BuildGraphScriptPredefinedTasks.INT.udn
+++ b/Engine/Documentation/Source/Programming/Development/BuildGraph/ScriptAnatomy/Tasks/PredefinedTasks/BuildGraphScriptPredefinedTasks.INT.udn
@@ -511,14 +511,18 @@ Creates a new changelist and submits a set of files to a Perforce stream.
 
 Task that strips symbols from a set of files.
 
-|              |                      |          |                                                                                                                                               |
-| ------------ | -------------------- | -------- | --------------------------------------------------------------------------------------------------------------------------------------------- |
-| Platform     | UnrealTargetPlatform | Required | The platform toolchain required to handle symbol files.                                                                                       |
-| Files        | String               | Required | List of output files. PDBs will be extracted from this list.                                                                                  |
-| StoreDir     | String               | Required | Output directory for the compressed symbols.                                                                                                  |
-| Product      | String               | Required | Name of the product for the symbol store records.                                                                                             |
-| BuildVersion | String               | Optional | BuildVersion associated with these symbols. Used for clean-up in AgeStore by matching this version against a directory name in a build share. |
-| If           | Condition            | Optional | Whether to execute this task. It is ignored if this condition evaluates to false.                                                             |
+|              |                      |          |                                                                                                                                                                                                                                                                                                                                       |
+| ------------ | -------------------- | -------- | ------------------------------------------------------------------------------------------------------------------------------------------------------------------------------------------------------------------------------------------------------------------------------------------------------------------------------------- |
+| Platform     | UnrealTargetPlatform | Required | The platform toolchain required to handle symbol files.                                                                                                                                                                                                                                                                               |
+| Files        | String               | Required | List of output files. PDBs will be extracted from this list.                                                                                                                                                                                                                                                                          |
+| StoreDir     | String               | Required | Output directory for the compressed symbols.                                                                                                                                                                                                                                                                                          |
+| Product      | String               | Required | Name of the product for the symbol store records.                                                                                                                                                                                                                                                                                     |
+| Branch       | String               | Optional | Name of the Branch to base all the depot source files from. Used when IndexSources is true (may be used only on some platforms).                                                                                                                                                                                                      |
+| Change       | Integer              | Optional | Changelist to which all the depot source files have been synced to. Used when IndexSources is true (may be used only on some platforms).                                                                                                                                                                                              |
+| BuildVersion | String               | Optional | BuildVersion associated with these symbols. Used for clean-up in AgeStore by matching this version against a directory name in a build share.                                                                                                                                                                                         |
+| IndexSources | Boolean              | Optional | Whether to include the source code index in the uploaded symbols. When enabled, the task will generate data required by a source server (only some platforms and source control servers are supported). The source server allows debuggers to automatically fetch the matching source code when debbugging builds or analyzing dumps. |
+| SourceFiles  | String               | Optional | Filter for the depot source files that are to be indexed. It's a semicolon-separated list of perforce filter e.g. Engine/....cpp;Engine/....h. It may also be a name of a previously defined tag e.g. "#SourceFiles Used when IndexSources is true (may be used only on some platforms).                                              |
+| If           | Condition            | Optional | Whether to execute this task. It is ignored if this condition evaluates to false.                                                                                                                                                                                                                                                     |
 
 ### Tag
 
diff --git a/Engine/Documentation/Source/Programming/UnrealBuildSystem/Configuration/BuildConfigProperties/BuildConfigProperties.INT.udn b/Engine/Documentation/Source/Programming/UnrealBuildSystem/Configuration/BuildConfigProperties/BuildConfigProperties.INT.udn
--- a/Engine/Documentation/Source/Programming/UnrealBuildSystem/Configuration/BuildConfigProperties/BuildConfigProperties.INT.udn
+++ b/Engine/Documentation/Source/Programming/UnrealBuildSystem/Configuration/BuildConfigProperties/BuildConfigProperties.INT.udn
@@ -134,7 +134,7 @@ $ bPGOProfile : Whether to enable Profile Guided Optimization (PGO) instrumentat
 
 $ bPGOOptimize : Whether to optimize this build with Profile Guided Optimization (PGO).
 
-$ bSupportEditAndContinue : Whether to support edit and continue.  Only works on Microsoft compilers.
+$ bSupportEditAndContinue : Whether to support edit and continue.
 
 $ bOmitFramePointers : Whether to omit frame pointers or not. Disabling is useful for e.g. memory profiling on the PC.
 
diff --git a/Engine/Documentation/Source/Programming/UnrealBuildSystem/TargetFiles/TargetFilesProperties/TargetFilesProperties.INT.udn b/Engine/Documentation/Source/Programming/UnrealBuildSystem/TargetFiles/TargetFilesProperties/TargetFilesProperties.INT.udn
--- a/Engine/Documentation/Source/Programming/UnrealBuildSystem/TargetFiles/TargetFilesProperties/TargetFilesProperties.INT.udn
+++ b/Engine/Documentation/Source/Programming/UnrealBuildSystem/TargetFiles/TargetFilesProperties/TargetFilesProperties.INT.udn
@@ -274,7 +274,7 @@ $ bPGOProfile (Boolean): Whether to enable Profile Guided Optimization (PGO) ins
 
 $ bPGOOptimize (Boolean): Whether to optimize this build with Profile Guided Optimization (PGO).
 
-$ bSupportEditAndContinue (Boolean): Whether to support edit and continue.  Only works on Microsoft compilers.
+$ bSupportEditAndContinue (Boolean): Whether to support edit and continue.
 
 $ bOmitFramePointers (Boolean): Whether to omit frame pointers or not. Disabling is useful for e.g. memory profiling on the PC.
 
