Documentation for build tools

diff --git a/Engine/Documentation/Source/Programming/UnrealBuildSystem/ModuleFiles/ModuleFilesProperties/ModuleFilesProperties.INT.udn b/Engine/Documentation/Source/Programming/UnrealBuildSystem/ModuleFiles/ModuleFilesProperties/ModuleFilesProperties.INT.udn
--- a/Engine/Documentation/Source/Programming/UnrealBuildSystem/ModuleFiles/ModuleFilesProperties/ModuleFilesProperties.INT.udn
+++ b/Engine/Documentation/Source/Programming/UnrealBuildSystem/ModuleFiles/ModuleFilesProperties/ModuleFilesProperties.INT.udn
@@ -48,6 +48,8 @@ $ bEnableExceptions (Boolean): Enable exception handling
 
 $ bEnableObjCExceptions (Boolean): Enable objective C exception handling
 
+$ bEnableObjCAutomaticReferenceCounting (Boolean): Enable objective C automatic reference counting (ARC) If you set this to true you should not use shared PCHs for this module. The engine won't be extensively using ARC in the short term Not doing this will result in a compile errors because shared PCHs were compiled with different flags than consumer
+
 $ bEnableUndefinedIdentifierWarnings (Boolean): Enable warnings for using undefined identifiers in #if expressions
 
 $ bDisableStaticAnalysis (Boolean): Disable all static analysis - clang, msvc, pvs-studio.
