Documentation for build tools

diff --git a/Engine/Documentation/Source/Programming/UnrealBuildSystem/TargetFiles/TargetFilesProperties/TargetFilesProperties.INT.udn b/Engine/Documentation/Source/Programming/UnrealBuildSystem/TargetFiles/TargetFilesProperties/TargetFilesProperties.INT.udn
--- a/Engine/Documentation/Source/Programming/UnrealBuildSystem/TargetFiles/TargetFilesProperties/TargetFilesProperties.INT.udn
+++ b/Engine/Documentation/Source/Programming/UnrealBuildSystem/TargetFiles/TargetFilesProperties/TargetFilesProperties.INT.udn
@@ -138,6 +138,10 @@ $ bForceEnableObjCExceptions (Boolean): Enable exceptions for all modules.
 
 $ bForceEnableRTTI (Boolean): Enable RTTI for all modules.
 
+$ bEnablePIE (Boolean): Enable Position Independent Executable (PIE). Has an overhead cost
+
+$ bEnableStackProtection (Boolean): Enable Stack Protection. Has an overhead cost
+
 $ bCompileWithStatsWithoutEngine (Boolean): Whether to include stats support even without the engine.
 
 $ bCompileWithPluginSupport (Boolean): Whether to include plugin support.
