Update tooltips for the Outliner.

diff --git a/Engine/Documentation/Source/Shared/LevelEditor/LevelEditor.INT.udn b/Engine/Documentation/Source/Shared/LevelEditor/LevelEditor.INT.udn
--- a/Engine/Documentation/Source/Shared/LevelEditor/LevelEditor.INT.udn
+++ b/Engine/Documentation/Source/Shared/LevelEditor/LevelEditor.INT.udn
@@ -64,11 +64,23 @@ Engine/UI/LevelEditor/Details
 [/EXCERPT:DetailsTab]
 
 [EXCERPT:SceneOutlinerTab]
-###World Outliner
-The World Outliner panel displays all of the Actors within the scene in a hierarchical tree view. Actors can be selected and modified directly from the World Outliner. You can also use the Info drop down menu to display an extra column that shows Levels, Layers, or ID Names.
+### Outliner
+The Outliner lists all actors in the level by name, organized in a tree view. You can select actors in this list exactly as you would in the Viewport, and right-click them to access the same contextual actions.
+
+You can also:
+
+-	Filter the list to a subset of one or more actors you’re interested in. Use the filter dropdown to the left of the Search box to choose from a set of pre-made filters, or enter your own custom filter text in the Search box. (You can also create new filters from your custom searches.)
+
+-	View details about the listed actors in dedicated columns, such as their visibility in editor, modified status, and source control status. Right-click the header row to choose the details you want to see.
+
+-	Create folders to group related actors together in the hierarchy, and drag and drop actors between folders.
+
+-	Drag and drop actors on top of another actor to give them a parent-child relationship. When you move, rotate, or scale a parent actor, its children follow too.
+
+You can have up to four Outliner panels active at any time. They all contain the same tree of actors, but you can choose to apply different filters to each.
 
 [VAR:ToolTipFullLink]
-Engine/UI/LevelEditor/SceneOutliner
+outliner-in-unreal-engine
 [/VAR]
 [/EXCERPT:SceneOutlinerTab]
 
diff --git a/Engine/Source/Editor/LevelEditor/Private/SLevelEditor.cpp b/Engine/Source/Editor/LevelEditor/Private/SLevelEditor.cpp
--- a/Engine/Source/Editor/LevelEditor/Private/SLevelEditor.cpp
+++ b/Engine/Source/Editor/LevelEditor/Private/SLevelEditor.cpp
@@ -946,7 +946,7 @@ TSharedRef<SDockTab> SLevelEditor::SpawnLevelEditorTab( const FSpawnTabArgs& Arg
 		
 		TSharedRef<SDockTab> SceneOutlinerTab = SNew( SDockTab )
 			.Label( Label )
-			.ToolTip( IDocumentation::Get()->CreateToolTip( Label, nullptr, "Shared/LevelEditor", "SceneOutlinerTab" ) )
+			.ToolTip( IDocumentation::Get()->CreateToolTip(NSLOCTEXT("LevelEditor", "OutlinerTabTooltip", "The Outliner lists all actors in the level by name, organized in a tree view."), nullptr, "Shared/LevelEditor", "SceneOutlinerTab" ) )
 			.OnExtendContextMenu(this, &SLevelEditor::OnExtendSceneOutlinerTabContextMenu)
 			[
 				SNew(SBorder)
