Documentation for build tools

diff --git a/Engine/Documentation/Source/Programming/UnrealBuildSystem/Configuration/BuildConfigProperties/BuildConfigProperties.INT.udn b/Engine/Documentation/Source/Programming/UnrealBuildSystem/Configuration/BuildConfigProperties/BuildConfigProperties.INT.udn
--- a/Engine/Documentation/Source/Programming/UnrealBuildSystem/Configuration/BuildConfigProperties/BuildConfigProperties.INT.udn
+++ b/Engine/Documentation/Source/Programming/UnrealBuildSystem/Configuration/BuildConfigProperties/BuildConfigProperties.INT.udn
@@ -86,6 +86,8 @@ $ bUseFastMonoCalls : New Monolithic Graphics drivers have optional "fast calls"
 
 $ NumIncludedBytesPerUnityCPP : An approximate number of bytes of C++ code to target for inclusion in a single unified C++ file.
 
+$ bDisableModuleNumIncludedBytesPerUnityCPPOverride : Disables overrides that are set by the module
+
 $ bStressTestUnity : Whether to stress test the C++ unity build robustness by including all C++ files files in a project from a single unified file.
 
 $ bDetailedUnityFiles : Whether to add additional information to the unity files, such as '_of_X' in the file name.
diff --git a/Engine/Documentation/Source/Programming/UnrealBuildSystem/TargetFiles/TargetFilesProperties/TargetFilesProperties.INT.udn b/Engine/Documentation/Source/Programming/UnrealBuildSystem/TargetFiles/TargetFilesProperties/TargetFilesProperties.INT.udn
--- a/Engine/Documentation/Source/Programming/UnrealBuildSystem/TargetFiles/TargetFilesProperties/TargetFilesProperties.INT.udn
+++ b/Engine/Documentation/Source/Programming/UnrealBuildSystem/TargetFiles/TargetFilesProperties/TargetFilesProperties.INT.udn
@@ -226,6 +226,8 @@ $ bUseFastMonoCalls (Boolean): New Monolithic Graphics drivers have optional "fa
 
 $ NumIncludedBytesPerUnityCPP (Int32): An approximate number of bytes of C++ code to target for inclusion in a single unified C++ file.
 
+$ bDisableModuleNumIncludedBytesPerUnityCPPOverride (Boolean): Disables overrides that are set by the module
+
 $ bStressTestUnity (Boolean): Whether to stress test the C++ unity build robustness by including all C++ files files in a project from a single unified file.
 
 $ bDetailedUnityFiles (Boolean): Whether to add additional information to the unity files, such as '_of_X' in the file name.
