Documentation for build tools

diff --git a/Engine/Documentation/Source/Programming/UnrealBuildSystem/ModuleFiles/ModuleFilesProperties/ModuleFilesProperties.INT.udn b/Engine/Documentation/Source/Programming/UnrealBuildSystem/ModuleFiles/ModuleFilesProperties/ModuleFilesProperties.INT.udn
--- a/Engine/Documentation/Source/Programming/UnrealBuildSystem/ModuleFiles/ModuleFilesProperties/ModuleFilesProperties.INT.udn
+++ b/Engine/Documentation/Source/Programming/UnrealBuildSystem/ModuleFiles/ModuleFilesProperties/ModuleFilesProperties.INT.udn
@@ -118,6 +118,8 @@ $ PublicRuntimeLibraryPaths (List&lt;String&gt;): List of search paths for libra
 
 $ PublicAdditionalLibraries (List&lt;String&gt;): List of additional libraries (names of the .lib files including extension) - typically used for External (third party) modules
 
+$ DependenciesToSkipPerArchitecture (Dictionary&lt;String, List&lt;String&gt;&gt;): Per-architecture lists of dependencies for linking to ignore (useful when building for multiple architectures, and a lib only is needed for one architecture), it's up to the Toolchain to use this
+
 $ PublicPreBuildLibraries (List&lt;String&gt;): List of additional pre-build libraries (names of the .lib files including extension) - typically used for additional targets which are still built, but using either TargetRules.PreBuildSteps or TargetRules.PreBuildTargets.
 
 $ PublicSystemLibraries (List&lt;String&gt;): List of system libraries to use - these are typically referenced via name and then found via the system paths. If you need to reference a .lib file use the PublicAdditionalLibraries instead
