Documentation for build tools

diff --git a/Engine/Documentation/Source/Programming/UnrealBuildSystem/Configuration/BuildConfigProperties/BuildConfigProperties.INT.udn b/Engine/Documentation/Source/Programming/UnrealBuildSystem/Configuration/BuildConfigProperties/BuildConfigProperties.INT.udn
--- a/Engine/Documentation/Source/Programming/UnrealBuildSystem/Configuration/BuildConfigProperties/BuildConfigProperties.INT.udn
+++ b/Engine/Documentation/Source/Programming/UnrealBuildSystem/Configuration/BuildConfigProperties/BuildConfigProperties.INT.udn
@@ -262,6 +262,8 @@ $ PCHMemoryAllocationFactor : Determines the amount of memory that the compiler
 
 $ AdditionalLinkerOptions : Allow the target to specify extra options for linking that aren't otherwise noted here
 
+$ bUseFastGenProfile : If specified along with -PGOProfile, then /FASTGENPROFILE will be used instead of /GENPROFILE. This usually means that the PGO data is generated faster, but the resulting data may not yield as efficient optimizations during -PGOOptimize
+
 $ bClangTimeTrace : (Experimental) Appends the -ftime-trace argument to the command line for Clang to output a JSON file containing a timeline for the compile. See http://aras-p.info/blog/2019/01/16/time-trace-timeline-flame-chart-profiler-for-Clang/ for more info.
 
 $ bCompilerTrace : Outputs compile timing information so that it can be analyzed.
