Documentation for build tools

diff --git a/Engine/Documentation/Source/Programming/UnrealBuildSystem/TargetFiles/TargetFilesProperties/TargetFilesProperties.INT.udn b/Engine/Documentation/Source/Programming/UnrealBuildSystem/TargetFiles/TargetFilesProperties/TargetFilesProperties.INT.udn
--- a/Engine/Documentation/Source/Programming/UnrealBuildSystem/TargetFiles/TargetFilesProperties/TargetFilesProperties.INT.udn
+++ b/Engine/Documentation/Source/Programming/UnrealBuildSystem/TargetFiles/TargetFilesProperties/TargetFilesProperties.INT.udn
@@ -398,7 +398,9 @@ $ IncludeOrderVersion (EngineIncludeOrderVersion): What version of include order
 
 $ bAllowHotReload (Boolean): Whether this target supports hot reload
 
-$ bUseGameplayDebugger (Boolean): Whether to compile with WITH_GAMEPLAY_DEBUGGER enabled and use GameplayDebugger.
+$ bUseGameplayDebugger (Boolean): Whether to compile with WITH_GAMEPLAY_DEBUGGER enabled with all Engine's default gameplay debugger categories.
+
+$ bUseGameplayDebuggerCore (Boolean): Set to true when bUseGameplayDebugger is false but GameplayDebugger's core parts are required.
 
 $ bBuildEditor (Boolean): Whether we are compiling editor code or not. Prefer the more explicit bCompileAgainstEditor instead.
 
