Fix link to docs for the PlayerStart actor.

diff --git a/Engine/Documentation/Source/Shared/Types/APlayerStart/APlayerStart.INT.udn b/Engine/Documentation/Source/Shared/Types/APlayerStart/APlayerStart.INT.udn
--- a/Engine/Documentation/Source/Shared/Types/APlayerStart/APlayerStart.INT.udn
+++ b/Engine/Documentation/Source/Shared/Types/APlayerStart/APlayerStart.INT.udn
@@ -4,10 +4,16 @@ Crumbs:
 Description:Overview of the PlayerStart class
 
 [EXCERPT:APlayerStart]
+The Player Start actor specifies a possible starting point for a player.
 
-A Player Start is an Actor that is placed in the level to designate where the player should start out when the game begins. 
+-   If the level contains more than one Player Start actor, the engine chooses one at random.
+
+-   The newly spawned player character adopts both the position and rotation of the chosen Player Start actor. The facing direction of the Player Start actor is shown in the editor viewport by the light blue arrow on its gizmo.
+
+-   If the Player Start gizmo shows a "bad size" warning in the editor viewport, move the actor so that it doesn't overlap any other objects in the level.
 
 [VAR:ToolTipFullLink]
-Engine/Actors/Types
+player-start-actor-in-unreal-engine
 [/VAR]
 [/EXCERPT:APlayerStart]
+
diff --git a/Engine/Source/Runtime/Engine/Classes/GameFramework/PlayerStart.h b/Engine/Source/Runtime/Engine/Classes/GameFramework/PlayerStart.h
--- a/Engine/Source/Runtime/Engine/Classes/GameFramework/PlayerStart.h
+++ b/Engine/Source/Runtime/Engine/Classes/GameFramework/PlayerStart.h
@@ -8,8 +8,9 @@
 #include "PlayerStart.generated.h"
 
 /** 
- *	This class indicates a location where a player can spawn when the game begins
+ *	This class indicates a location where a player can spawn when the game begins.
  *	
+ * 
  *	@see https://docs.unrealengine.com/latest/INT/Engine/Actors/PlayerStart/
  */
 UCLASS(Blueprintable, ClassGroup=Common, hidecategories=Collision)
