Documentation for build tools

diff --git a/Engine/Documentation/Source/Programming/UnrealBuildSystem/Configuration/BuildConfigProperties/BuildConfigProperties.INT.udn b/Engine/Documentation/Source/Programming/UnrealBuildSystem/Configuration/BuildConfigProperties/BuildConfigProperties.INT.udn
--- a/Engine/Documentation/Source/Programming/UnrealBuildSystem/Configuration/BuildConfigProperties/BuildConfigProperties.INT.udn
+++ b/Engine/Documentation/Source/Programming/UnrealBuildSystem/Configuration/BuildConfigProperties/BuildConfigProperties.INT.udn
@@ -324,7 +324,7 @@ $ ProcessorCountMultiplier : Processor count multiplier for local execution. Can
 
 $ MemoryPerActionBytes : Free memory per action in bytes, used to limit the number of parallel actions if the machine is memory starved. Set to 0 to disable free memory checking.
 
-$ ProcessPriority : The priority to set for spawned processes. Valid Settings: Idle, BelowNormal, Normal, AboveNormal, High Default: BelowNormal
+$ ProcessPriority : The priority to set for spawned processes. Valid Settings: Idle, BelowNormal, Normal, AboveNormal, High Default: BelowNormal or Normal for an Asymmetrical processor as BelowNormal can cause scheduling issues.
 
 $ bStopCompilationAfterErrors : When enabled, will stop compiling targets after a compile error occurs.
 
