Documentation for build tools

diff --git a/Engine/Documentation/Source/Programming/UnrealBuildSystem/ModuleFiles/ModuleFilesProperties/ModuleFilesProperties.INT.udn b/Engine/Documentation/Source/Programming/UnrealBuildSystem/ModuleFiles/ModuleFilesProperties/ModuleFilesProperties.INT.udn
--- a/Engine/Documentation/Source/Programming/UnrealBuildSystem/ModuleFiles/ModuleFilesProperties/ModuleFilesProperties.INT.udn
+++ b/Engine/Documentation/Source/Programming/UnrealBuildSystem/ModuleFiles/ModuleFilesProperties/ModuleFilesProperties.INT.udn
@@ -80,7 +80,7 @@ $ AllowedRestrictedFolders (List&lt;String&gt;): List of folders which are allow
 
 $ AliasRestrictedFolders (Dictionary&lt;String, String&gt;): Set of aliased restricted folder references
 
-$ bEnforceIWYU (Boolean): Enforce "include what you use" rules when PCHUsage is set to ExplicitOrSharedPCH; warns when monolithic headers (Engine.h, UnrealEd.h, etc...) are used, and checks that source files include their matching header first.
+$ IWYUSupport (IWYUSupport): Allows "include what you use" to modify the source code when run. bEnforceIWYU must be true for this variable to matter.
 
 $ bAddDefaultIncludePaths (Boolean): Whether to add all the default include paths to the module (eg. the Source/Classes folder, subfolders under Source/Public).
 
@@ -184,6 +184,8 @@ $ UnsafeTypeCastWarningLevel (WarningLevel): How to treat unsafe implicit type c
 
 $ bUseUnity (Boolean): If unity builds are enabled this can be used to override if this specific module will build using Unity. This is set using the per module configurations in BuildConfiguration.
 
+$ bEnforceIWYU (Boolean): Enforce "include what you use" rules when PCHUsage is set to ExplicitOrSharedPCH; warns when monolithic headers (Engine.h, UnrealEd.h, etc...) are used, and checks that source files include their matching header first.
+
 $ bLegacyPublicIncludePaths (Boolean): Whether this module qualifies included headers from other modules relative to the root of their 'Public' folder. This reduces the number of search paths that have to be passed to the compiler, improving performance and reducing the length of the compiler command line.
 
 $ bLegacyParentIncludePaths (Boolean): Whether this module qualifies included headers from other modules relative to the parent directory. This reduces the number of search paths that have to be passed to the compiler, improving performance and reducing the length of the compiler command line.
