Documentation for build tools

diff --git a/Engine/Documentation/Source/Programming/UnrealBuildSystem/TargetFiles/TargetFilesProperties/TargetFilesProperties.INT.udn b/Engine/Documentation/Source/Programming/UnrealBuildSystem/TargetFiles/TargetFilesProperties/TargetFilesProperties.INT.udn
--- a/Engine/Documentation/Source/Programming/UnrealBuildSystem/TargetFiles/TargetFilesProperties/TargetFilesProperties.INT.udn
+++ b/Engine/Documentation/Source/Programming/UnrealBuildSystem/TargetFiles/TargetFilesProperties/TargetFilesProperties.INT.udn
@@ -168,6 +168,8 @@ $ bUseEstimatedUtcNow (Boolean): Whether to use the EstimatedUtcNow or PlatformU
 
 $ bCompileFreeType (Boolean): True if we need FreeType support.
 
+$ bUseExecCommandsInShipping (Boolean): Whether to turn allow exec commands for shipping builds.
+
 $ bCompileForSize (Boolean): True if we want to favor optimizing size over speed.
 
 $ OptimizationLevel (OptimizationMode): Allows to fine tune optimizations level for speed and\or code size
