Documentation for build tools

diff --git a/Engine/Documentation/Source/Programming/UnrealBuildSystem/ModuleFiles/ModuleFilesProperties/ModuleFilesProperties.INT.udn b/Engine/Documentation/Source/Programming/UnrealBuildSystem/ModuleFiles/ModuleFilesProperties/ModuleFilesProperties.INT.udn
--- a/Engine/Documentation/Source/Programming/UnrealBuildSystem/ModuleFiles/ModuleFilesProperties/ModuleFilesProperties.INT.udn
+++ b/Engine/Documentation/Source/Programming/UnrealBuildSystem/ModuleFiles/ModuleFilesProperties/ModuleFilesProperties.INT.udn
@@ -54,6 +54,10 @@ $ bEnableUndefinedIdentifierWarnings (Boolean): Enable warnings for using undefi
 
 $ bDisableStaticAnalysis (Boolean): Disable all static analysis - clang, msvc, pvs-studio.
 
+$ bStaticAnalyzerExtensions (Boolean): Enable additional analyzer extension warnings using the EspXEngine plugin. This is only supported for MSVC. See https://learn.microsoft.com/en-us/cpp/code-quality/using-the-cpp-core-guidelines-checkers This will add a large number of warnings by default. It's recommended to use StaticAnalyzerRulesets if this is enabled.
+
+$ StaticAnalyzerRulesets (HashSet&lt;FileReference&gt;): The static analyzer rulesets that should be used to filter warnings. This is only supported for MSVC. See https://learn.microsoft.com/en-us/cpp/code-quality/using-rule-sets-to-specify-the-cpp-rules-to-run
+
 $ StaticAnalyzerCheckers (HashSet&lt;String&gt;): The static analyzer checkers that should be enabled rather than the defaults. This is only supported for Clang. See https://clang.llvm.org/docs/analyzer/checkers.html for a full list. Or run: 'clang -Xclang -analyzer-checker-help' or: 'clang -Xclang -analyzer-checker-help-alpha' for the list of experimental checkers.
 
 $ StaticAnalyzerDisabledCheckers (HashSet&lt;String&gt;): The static analyzer default checkers that should be disabled. Unused if StaticAnalyzerCheckers is populated. This is only supported for Clang. This overrides the default disabled checkers, which are deadcode.DeadStores and security.FloatLoopCounter See https://clang.llvm.org/docs/analyzer/checkers.html for a full list. Or run: 'clang -Xclang -analyzer-checker-help' or: 'clang -Xclang -analyzer-checker-help-alpha' for the list of experimental checkers.
