Documentation for build tools

diff --git a/Engine/Documentation/Source/Programming/UnrealBuildSystem/ModuleFiles/ModuleFilesProperties/ModuleFilesProperties.INT.udn b/Engine/Documentation/Source/Programming/UnrealBuildSystem/ModuleFiles/ModuleFilesProperties/ModuleFilesProperties.INT.udn
--- a/Engine/Documentation/Source/Programming/UnrealBuildSystem/ModuleFiles/ModuleFilesProperties/ModuleFilesProperties.INT.udn
+++ b/Engine/Documentation/Source/Programming/UnrealBuildSystem/ModuleFiles/ModuleFilesProperties/ModuleFilesProperties.INT.udn
@@ -122,7 +122,7 @@ $ PublicRuntimeLibraryPaths (List&lt;String&gt;): List of search paths for libra
 
 $ PublicAdditionalLibraries (List&lt;String&gt;): List of additional libraries (names of the .lib files including extension) - typically used for External (third party) modules
 
-$ DependenciesToSkipPerArchitecture (Dictionary&lt;String, List&lt;String&gt;&gt;): Per-architecture lists of dependencies for linking to ignore (useful when building for multiple architectures, and a lib only is needed for one architecture), it's up to the Toolchain to use this
+$ DependenciesToSkipPerArchitecture (Dictionary&lt;String, List&lt;UnrealArch&gt;&gt;): Per-architecture lists of dependencies for linking to ignore (useful when building for multiple architectures, and a lib only is needed for one architecture), it's up to the Toolchain to use this
 
 $ PublicPreBuildLibraries (List&lt;String&gt;): List of additional pre-build libraries (names of the .lib files including extension) - typically used for additional targets which are still built, but using either TargetRules.PreBuildSteps or TargetRules.PreBuildTargets.
 
diff --git a/Engine/Documentation/Source/Programming/UnrealBuildSystem/TargetFiles/TargetFilesProperties/TargetFilesProperties.INT.udn b/Engine/Documentation/Source/Programming/UnrealBuildSystem/TargetFiles/TargetFilesProperties/TargetFilesProperties.INT.udn
--- a/Engine/Documentation/Source/Programming/UnrealBuildSystem/TargetFiles/TargetFilesProperties/TargetFilesProperties.INT.udn
+++ b/Engine/Documentation/Source/Programming/UnrealBuildSystem/TargetFiles/TargetFilesProperties/TargetFilesProperties.INT.udn
@@ -10,7 +10,7 @@ $ Platform (UnrealTargetPlatform): Platform that this target is being built for.
 
 $ Configuration (UnrealTargetConfiguration): The configuration being built.
 
-$ Architecture (String): Architecture that the target is being built for (or an empty string for the default).
+$ Architectures (UnrealArchitectures): Architecture that the target is being built for (or an empty string for the default).
 
 $ ProjectFile (FileReference): Path to the project file for the project containing this target.
 
@@ -22,6 +22,8 @@ $ ExplicitTestsTarget (Boolean): Whether this is a test target explicitly define
 
 $ WithLowLevelTests (Boolean): Controls the value of WITH_LOW_LEVEL_TESTS that dictates whether module-specific low level tests are compiled in or not.
 
+$ Architecture (UnrealArch): Gets the Architecture in the normal case where there is a single Architecture in Architectures (this will throw an exception if there is more than one architecture specified)
+
 $ HostPlatform (UnrealTargetPlatform): Gets the host platform being built on
 
 $ bGenerateProjectFiles (Boolean): Expose the bGenerateProjectFiles flag to targets, so we can modify behavior as appropriate for better intellisense
