Documentation for build tools

diff --git a/Engine/Documentation/Source/Programming/Development/BuildGraph/ScriptAnatomy/Tasks/PredefinedTasks/BuildGraphScriptPredefinedTasks.INT.udn b/Engine/Documentation/Source/Programming/Development/BuildGraph/ScriptAnatomy/Tasks/PredefinedTasks/BuildGraphScriptPredefinedTasks.INT.udn
--- a/Engine/Documentation/Source/Programming/Development/BuildGraph/ScriptAnatomy/Tasks/PredefinedTasks/BuildGraphScriptPredefinedTasks.INT.udn
+++ b/Engine/Documentation/Source/Programming/Development/BuildGraph/ScriptAnatomy/Tasks/PredefinedTasks/BuildGraphScriptPredefinedTasks.INT.udn
@@ -106,16 +106,17 @@ Compiles a target with UnrealBuildTool.
 
 Cook a selection of maps for a certain platform
 
-|           |           |          |                                                                                   |
-| --------- | --------- | -------- | --------------------------------------------------------------------------------- |
-| Project   | String    | Required | Project file to be cooked.                                                        |
-| Platform  | String    | Required | The cook platform to target (for example, Windows).                               |
-| Maps      | String    | Optional | List of maps to be cooked, separated by '+' characters.                           |
-| Versioned | Boolean   | Optional | Additional arguments to be passed to the cooker.                                  |
-| Arguments | String    | Optional | Additional arguments to be passed to the cooker.                                  |
-| EditorExe | String    | Optional | Optional path to what editor executable to run for cooking.                       |
-| Tag       | TagList   | Optional | Tag to be applied to build products of this task.                                 |
-| If        | Condition | Optional | Whether to execute this task. It is ignored if this condition evaluates to false. |
+|           |           |          |                                                                                                                                                                          |
+| --------- | --------- | -------- | ------------------------------------------------------------------------------------------------------------------------------------------------------------------------ |
+| Project   | String    | Required | Project file to be cooked.                                                                                                                                               |
+| Platform  | String    | Required | The cook platform to target (for example, Windows).                                                                                                                      |
+| Maps      | String    | Optional | List of maps to be cooked, separated by '+' characters.                                                                                                                  |
+| Versioned | Boolean   | Optional | Additional arguments to be passed to the cooker.                                                                                                                         |
+| Arguments | String    | Optional | Additional arguments to be passed to the cooker.                                                                                                                         |
+| EditorExe | String    | Optional | Optional path to what editor executable to run for cooking.                                                                                                              |
+| TagOutput | Boolean   | Optional | Whether to tag the output from the cook. Since cooks produce a lot of files, it can be detrimental to spend time tagging them if we don't need them in a dependent node. |
+| Tag       | TagList   | Optional | Tag to be applied to build products of this task.                                                                                                                        |
+| If        | Condition | Optional | Whether to execute this task. It is ignored if this condition evaluates to false.                                                                                        |
 
 ### Copy
 
