Documentation for build tools

diff --git a/Engine/Documentation/Source/Programming/UnrealBuildSystem/Configuration/BuildConfigProperties/BuildConfigProperties.INT.udn b/Engine/Documentation/Source/Programming/UnrealBuildSystem/Configuration/BuildConfigProperties/BuildConfigProperties.INT.udn
--- a/Engine/Documentation/Source/Programming/UnrealBuildSystem/Configuration/BuildConfigProperties/BuildConfigProperties.INT.udn
+++ b/Engine/Documentation/Source/Programming/UnrealBuildSystem/Configuration/BuildConfigProperties/BuildConfigProperties.INT.udn
@@ -776,8 +776,6 @@ $ SshAuthentication : The authentication used for SSH (probably similar to Rsync
 
 ### HoloLensPlatform
 
-$ Compiler : Version of the compiler toolchain to use on HoloLens. A value of "default" will be changed to a specific version at UBT startup.
-
 ### Log
 
 $ bBackupLogFiles : Whether to backup an existing log file, rather than overwriting it.
