Documentation for build tools

diff --git a/Engine/Documentation/Source/Programming/UnrealBuildSystem/Configuration/BuildConfigProperties/BuildConfigProperties.INT.udn b/Engine/Documentation/Source/Programming/UnrealBuildSystem/Configuration/BuildConfigProperties/BuildConfigProperties.INT.udn
--- a/Engine/Documentation/Source/Programming/UnrealBuildSystem/Configuration/BuildConfigProperties/BuildConfigProperties.INT.udn
+++ b/Engine/Documentation/Source/Programming/UnrealBuildSystem/Configuration/BuildConfigProperties/BuildConfigProperties.INT.udn
@@ -776,6 +776,8 @@ $ SshAuthentication : The authentication used for SSH (probably similar to Rsync
 
 ### HoloLensPlatform
 
+$ Compiler : Version of the compiler toolchain to use on HoloLens. A value of "default" will be changed to a specific version at UBT startup.
+
 ### Log
 
 $ bBackupLogFiles : Whether to backup an existing log file, rather than overwriting it.
diff --git a/Engine/Documentation/Source/Programming/UnrealBuildSystem/ModuleFiles/ModuleFilesProperties/ModuleFilesProperties.INT.udn b/Engine/Documentation/Source/Programming/UnrealBuildSystem/ModuleFiles/ModuleFilesProperties/ModuleFilesProperties.INT.udn
--- a/Engine/Documentation/Source/Programming/UnrealBuildSystem/ModuleFiles/ModuleFilesProperties/ModuleFilesProperties.INT.udn
+++ b/Engine/Documentation/Source/Programming/UnrealBuildSystem/ModuleFiles/ModuleFilesProperties/ModuleFilesProperties.INT.udn
@@ -176,6 +176,8 @@ $ OverridePackageType (PackageOverrideType): Overridden type of module that will
 
 $ OptimizeCode (CodeOptimization): When this module's code should be optimized.
 
+$ OptimizationLevel (OptimizationMode): Allows fine tuning optimization level for speed and\or code size. This requires a private PCH (or NoPCHs, which is not recommended)
+
 $ PCHUsage (PCHUsageMode): Precompiled header usage for this module
 
 $ DefaultBuildSettings (BuildSettingsVersion): Which engine version's build settings to use by default.
