Documentation for build tools

diff --git a/Engine/Documentation/Source/Programming/Development/BuildGraph/ScriptAnatomy/Tasks/PredefinedTasks/BuildGraphScriptPredefinedTasks.INT.udn b/Engine/Documentation/Source/Programming/Development/BuildGraph/ScriptAnatomy/Tasks/PredefinedTasks/BuildGraphScriptPredefinedTasks.INT.udn
--- a/Engine/Documentation/Source/Programming/Development/BuildGraph/ScriptAnatomy/Tasks/PredefinedTasks/BuildGraphScriptPredefinedTasks.INT.udn
+++ b/Engine/Documentation/Source/Programming/Development/BuildGraph/ScriptAnatomy/Tasks/PredefinedTasks/BuildGraphScriptPredefinedTasks.INT.udn
@@ -246,14 +246,14 @@ Spawns Docker and waits for it to complete.
 
 Compiles C# project files, and their dependencies.
 
-|           |               |          |                                                                                                |
-| --------- | ------------- | -------- | ---------------------------------------------------------------------------------------------- |
-| Paths     | String        | Required | The C# project file to compile. Using semicolons, more than one project file can be specified. |
-| Change    | Integer       | Optional | The configuration to compile.                                                                  |
-| MinChange | Integer       | Optional | The configuration to compile.                                                                  |
-| MaxChange | Integer       | Optional | The configuration to compile.                                                                  |
-| Output    | FileReference | Optional | The file to write to                                                                           |
-| If        | Condition     | Optional | Whether to execute this task. It is ignored if this condition evaluates to false.              |
+|           |               |          |                                                                                   |
+| --------- | ------------- | -------- | --------------------------------------------------------------------------------- |
+| Path      | String        | Optional | Filter Path (default is ...)                                                      |
+| Change    | Integer       | Optional | The configuration to compile.                                                     |
+| MinChange | Integer       | Optional | The configuration to compile.                                                     |
+| MaxChange | Integer       | Optional | The configuration to compile.                                                     |
+| Output    | FileReference | Optional | The file to write to                                                              |
+| If        | Condition     | Optional | Whether to execute this task. It is ignored if this condition evaluates to false. |
 
 ### Git
 
