Documentation for build tools

diff --git a/Engine/Documentation/Source/Programming/UnrealBuildSystem/Configuration/BuildConfigProperties/BuildConfigProperties.INT.udn b/Engine/Documentation/Source/Programming/UnrealBuildSystem/Configuration/BuildConfigProperties/BuildConfigProperties.INT.udn
--- a/Engine/Documentation/Source/Programming/UnrealBuildSystem/Configuration/BuildConfigProperties/BuildConfigProperties.INT.udn
+++ b/Engine/Documentation/Source/Programming/UnrealBuildSystem/Configuration/BuildConfigProperties/BuildConfigProperties.INT.udn
@@ -308,6 +308,8 @@ $ bStopOnError : Whether to stop on error
 
 $ MsvcCRTRedistVersion : Which MSVC CRT Redist version to use
 
+$ CompilerVersion : Which MSVC Compiler version to use
+
 ### HordeExecutor
 
 $ NumRemoteParallelProcesses : How many processes that will be executed in parallel remotely
