Documentation for build tools

diff --git a/Engine/Documentation/Source/Programming/UnrealBuildSystem/Configuration/BuildConfigProperties/BuildConfigProperties.INT.udn b/Engine/Documentation/Source/Programming/UnrealBuildSystem/Configuration/BuildConfigProperties/BuildConfigProperties.INT.udn
--- a/Engine/Documentation/Source/Programming/UnrealBuildSystem/Configuration/BuildConfigProperties/BuildConfigProperties.INT.udn
+++ b/Engine/Documentation/Source/Programming/UnrealBuildSystem/Configuration/BuildConfigProperties/BuildConfigProperties.INT.udn
@@ -102,12 +102,14 @@ $ bUsePDBFiles : Whether PDB files should be used for Visual C++ builds.
 
 $ bUsePCHFiles : Whether PCH files should be used.
 
-$ bDeterministic : Set flags require for determinstic linking (experimental, may not be fully supported). Deterministic compiling is controlled via ModuleRules.
+$ bDeterministic : Set flags require for deterministic linking (experimental, may not be fully supported). Deterministic compiling is controlled via ModuleRules.
 
-$ bForceDeterministic : Force set flags require for determinstic compiling and linking (experimental, may not be fully supported). This setting is only recommended for testing, instead:
+$ bForceDeterministic : Force set flags require for deterministic compiling and linking (experimental, may not be fully supported). This setting is only recommended for testing, instead:
         * Set bDeterministic on a per module basis in ModuleRules to control deterministic compiling.
         * Set bDeterministic on a per target basis in TargetRules to control deterministic linking.
 
+$ bChainPCHs : Whether PCHs should be chained when compiling with clang.
+
 $ bForceIncludePCHHeadersForGenCppFilesWhenPCHIsDisabled : Whether PCH headers should be force included for gen.cpp files when PCH is disabled.
 
 $ bPreprocessDepends : Generate dependency files by preprocessing. This is only recommended when distributing builds as it adds additional overhead.
diff --git a/Engine/Documentation/Source/Programming/UnrealBuildSystem/TargetFiles/TargetFilesProperties/TargetFilesProperties.INT.udn b/Engine/Documentation/Source/Programming/UnrealBuildSystem/TargetFiles/TargetFilesProperties/TargetFilesProperties.INT.udn
--- a/Engine/Documentation/Source/Programming/UnrealBuildSystem/TargetFiles/TargetFilesProperties/TargetFilesProperties.INT.udn
+++ b/Engine/Documentation/Source/Programming/UnrealBuildSystem/TargetFiles/TargetFilesProperties/TargetFilesProperties.INT.udn
@@ -252,12 +252,14 @@ $ bUsePDBFiles (Boolean): Whether PDB files should be used for Visual C++ builds
 
 $ bUsePCHFiles (Boolean): Whether PCH files should be used.
 
-$ bDeterministic (Boolean): Set flags require for determinstic linking (experimental, may not be fully supported). Deterministic compiling is controlled via ModuleRules.
+$ bDeterministic (Boolean): Set flags require for deterministic linking (experimental, may not be fully supported). Deterministic compiling is controlled via ModuleRules.
 
-$ bForceDeterministic (Boolean): Force set flags require for determinstic compiling and linking (experimental, may not be fully supported). This setting is only recommended for testing, instead:
+$ bForceDeterministic (Boolean): Force set flags require for deterministic compiling and linking (experimental, may not be fully supported). This setting is only recommended for testing, instead:
         * Set bDeterministic on a per module basis in ModuleRules to control deterministic compiling.
         * Set bDeterministic on a per target basis in TargetRules to control deterministic linking.
 
+$ bChainPCHs (Boolean): Whether PCHs should be chained when compiling with clang.
+
 $ bForceIncludePCHHeadersForGenCppFilesWhenPCHIsDisabled (Boolean): Whether PCH headers should be force included for gen.cpp files when PCH is disabled.
 
 $ bPreprocessOnly (Boolean): Whether to just preprocess source files for this target, and skip compilation
