Documentation for build tools

diff --git a/Engine/Documentation/Source/Programming/UnrealBuildSystem/Configuration/BuildConfigProperties/BuildConfigProperties.INT.udn b/Engine/Documentation/Source/Programming/UnrealBuildSystem/Configuration/BuildConfigProperties/BuildConfigProperties.INT.udn
--- a/Engine/Documentation/Source/Programming/UnrealBuildSystem/Configuration/BuildConfigProperties/BuildConfigProperties.INT.udn
+++ b/Engine/Documentation/Source/Programming/UnrealBuildSystem/Configuration/BuildConfigProperties/BuildConfigProperties.INT.udn
@@ -238,6 +238,10 @@ $ MaxRootPathLength : Maximum recommended root path length.
 
 $ MaxNestedPathLength : Maximum length of a path relative to the root directory. Used on Windows to ensure paths are portable between machines. Defaults to off.
 
+$ bUseFastGenProfile : If specified along with -PGOProfile, then /FASTGENPROFILE will be used instead of /GENPROFILE. This usually means that the PGO data is generated faster, but the resulting data may not yield as efficient optimizations during -PGOOptimize
+
+$ bPGONoExtraCounters : If specified along with -PGOProfile, prevent the usage of extra counters. Please note that by default /FASTGENPROFILE doesnt use extra counters https://learn.microsoft.com/en-us/cpp/build/reference/genprofile-fastgenprofile-generate-profiling-instrumented-build?view=msvc-170
+
 $ Compiler : Version of the compiler toolchain to use on Windows platform. A value of "default" will be changed to a specific version at UBT start up.
 
 $ CompilerVersion : The specific toolchain version to use. This may be a specific version number (for example, "14.13.26128"), the string "Latest" to select the newest available version, or the string "Preview" to select the newest available preview version. By default, and if it is available, we use the toolchain version indicated by WindowsPlatform.DefaultToolChainVersion (otherwise, we use the latest version).
@@ -264,8 +268,6 @@ $ PCHMemoryAllocationFactor : Determines the amount of memory that the compiler
 
 $ AdditionalLinkerOptions : Allow the target to specify extra options for linking that aren't otherwise noted here
 
-$ bUseFastGenProfile : If specified along with -PGOProfile, then /FASTGENPROFILE will be used instead of /GENPROFILE. This usually means that the PGO data is generated faster, but the resulting data may not yield as efficient optimizations during -PGOOptimize
-
 $ bClangTimeTrace : (Experimental) Appends the -ftime-trace argument to the command line for Clang to output a JSON file containing a timeline for the compile. See http://aras-p.info/blog/2019/01/16/time-trace-timeline-flame-chart-profiler-for-Clang/ for more info.
 
 $ bCompilerTrace : Outputs compile timing information so that it can be analyzed.
