Documentation for build tools

diff --git a/Engine/Documentation/Source/Programming/UnrealBuildSystem/Configuration/BuildConfigProperties/BuildConfigProperties.INT.udn b/Engine/Documentation/Source/Programming/UnrealBuildSystem/Configuration/BuildConfigProperties/BuildConfigProperties.INT.udn
--- a/Engine/Documentation/Source/Programming/UnrealBuildSystem/Configuration/BuildConfigProperties/BuildConfigProperties.INT.udn
+++ b/Engine/Documentation/Source/Programming/UnrealBuildSystem/Configuration/BuildConfigProperties/BuildConfigProperties.INT.udn
@@ -286,6 +286,10 @@ $ EnableOptimizeCode : List of modules to enable optimizations for
 
 $ DisableOptimizeCode : List of modules to disable optimizations for
 
+$ OptimizeForSize : List of modules to optimize for size. This allows the target to override module optimization level Note that this may disable PCH usage if a private PCH is not provided
+
+$ OptimizeForSizeAndSpeed : List of modules to optimize for size and speed. This allows the target to override module optimization level Note that this may disable PCH usage if a private PCH is not provided
+
 ### TestTargetRules
 
 $ bDisableInliningGenCpps : Disable supports for inlining gen.cpps
diff --git a/Engine/Documentation/Source/Programming/UnrealBuildSystem/TargetFiles/TargetFilesProperties/TargetFilesProperties.INT.udn b/Engine/Documentation/Source/Programming/UnrealBuildSystem/TargetFiles/TargetFilesProperties/TargetFilesProperties.INT.udn
--- a/Engine/Documentation/Source/Programming/UnrealBuildSystem/TargetFiles/TargetFilesProperties/TargetFilesProperties.INT.udn
+++ b/Engine/Documentation/Source/Programming/UnrealBuildSystem/TargetFiles/TargetFilesProperties/TargetFilesProperties.INT.udn
@@ -396,6 +396,10 @@ $ EnableOptimizeCodeForModules (String[]): List of modules to enable optimizatio
 
 $ DisableOptimizeCodeForModules (String[]): List of modules to disable optimizations for
 
+$ OptimizeForSizeModules (String[]): List of modules to optimize for size. This allows the target to override module optimization level Note that this may disable PCH usage if a private PCH is not provided
+
+$ OptimizeForSizeAndSpeedModules (String[]): List of modules to optimize for size and speed. This allows the target to override module optimization level Note that this may disable PCH usage if a private PCH is not provided
+
 $ GeneratedProjectName (String): When generating project files, specifies the name of the project file to use when there are multiple targets of the same type.
 
 $ OptedInModulePlatforms (UnrealTargetPlatform[]): If this is non-null, then any platforms NOT listed will not be allowed to have modules in their directories be created
