Improve collision tooltips.

diff --git a/Engine/Documentation/Source/Shared/Collision/CollisionTooltips.INT.udn b/Engine/Documentation/Source/Shared/Collision/CollisionTooltips.INT.udn
--- a/Engine/Documentation/Source/Shared/Collision/CollisionTooltips.INT.udn
+++ b/Engine/Documentation/Source/Shared/Collision/CollisionTooltips.INT.udn
@@ -6,32 +6,40 @@ Description: Collision Help
 
 [EXCERPT:ObjectChannel]
 
-Object Channel is the collision channel type that can be assigned to object as Object Type. 
+Object channels are used when the engine tests for collision between moving objects, and when you run your own collision queries against objects.
 
-![](ObjectType.png)
+By creating your own custom object channels, you can define patterns of physical interactions between different types of actors, and you can filter collision queries to find only actors that you set to respond to those custom channels.
 
-This used to be [Movement Channel] because it is used when this object is moved. 
+---
 
-When you define collision property for object, you define 3 set ups.
+Every actor that can collide with anything else has an object channel assigned to it in the Object Type setting of the actor's collision preset:
 
-- Collision Enabled : Is there any collision? If there is collision, do you want this object to be simulated?
-- Object Type : Define what object type this will use when it is being moved
-- Collision Response : Response to other objects
+![](CollisionPresetObjectType.png)
 
-![](CollisionResponse.png)
+---
 
-This is a lot of information to set up, so we implemented [Collision Presets] as the one you can quickly choose and set. 
+When an actor moves in space so that it collides with another actor, the collision responses in each actor's collision preset are checked against the other actor's Object Type to determine what happens:
 
-![](CollisionPreset.png)
+-   If both presets are set to Block the other actor's object type, the movement is blocked. In addition, if both actors have the Simulation Generates Hit Event setting enabled, a Hit event is generated for the collision.
 
-If you would like to assign to object, you want this channel to be [Object Channel]
+-   Otherwise, the movement is permitted. In addition, if both presets are set to Overlap the other actor's object type AND both actors have the Generate Overlap Events setting enabled, Overlap events are generated for the collision.
 
-Query functions also allow trace [Object Type]. For example, you can query finding all objects within this sphere when the Object Type is Pawn. 
+---
 
-![](ObjectChannelQuery.png)
+Some types of collision queries, like traces and overlap tests, also allow you to specify the Object Types that you want to take into account in the query:
+
+![](ObjectQueryExample.png)
+
+If a query hits an actor, the object response in the collision preset determines what happens:
+
+-   Ignore: The query does not report a hit, and it continues looking for other collisions.
+
+-   Overlap: The query returns information about the collision in the Out Hit output, and it continues looking for other collisions.
+
+-   Block: The query returns information about the collision in the Out Hit output, and it stops looking for any other collisions.
 
 [VAR:ToolTipFullLink]
-Engine/Physics/Collision
+collision-in-unreal-engine---overview
 [/VAR]
 
 [/EXCERPT:ObjectChannel]
@@ -39,59 +47,72 @@ Engine/Physics/Collision
 
 [EXCERPT:TraceChannel]
 
-We distinguish Trace Channel vs Object Channel for major two reasons. 
+Trace channels determine how actors respond when you run a trace query.
 
-Object Channel is the one that can be assigned to object as Object Type and can be queried by Object Type when traced. Please check the above Object Channel tooltip. 
-Trace Channel is when you have a specific purpose of trace, for example, Weapon Trace. Even though it is possible to create bullet object that has object type as "Weapon", often most of games do direct ray cast to the target. When you trace using "Weapon" Trace type, each object can decide how you would like to respond to it. 
+By creating your own custom trace channels, you can filter your collision queries to find only actors that you set up to respond to those channels.
 
-Another example is camera. You can decide to have camera object to contain their Object Type to be Camera, in which case Camera is Object Channel. 
-But also you can choose to use Camera as Trace Channel, and use it in your trace code to trace using that channel, and other objects can decide how they would like to respond. 
+---
 
-The distinction is shown in below when using same functionality. 
+Every actor that can collide with anything else has a collision preset that defines how it responds to each trace channel: Ignore, Overlap, or Block. In this example, the preset blocks trace queries that use the Visibility channel, and ignores queries that use the Camera channel:
 
-![](TraceQuery.png)
+![](CollisionPresetTraceResponses.png)
 
-![](ObjectQuery.png)
+When you carry out a trace query, you can specify the trace channel that you want the collision test to use:
 
-This two works slightly different. 
+![](TraceQueryExample.png)
 
-First version is trace by Channel, which specifies Trace Channel. 
-This will trace from Start to End with sphere of radius 30, and will find anything that RESPONDS to Weapon Trace. 
+If a trace query hits an actor, the trace response in the collision preset determines what happens:
 
-This respond is defined per object. 
+-   Ignore: The query does not report a hit, and it continues looking for other collisions.
 
-![](CollisionResponse.png)
+-   Overlap: If you are performing a single trace, as in the example above, the query returns information about the collision in the Out Hit output and stops. If you are performing a Multi trace, the query returns information about the collision in the Out Hits array output and continues looking for other collisions.
 
-Second version is trace by Object. 
-This will trace from Start to End with sphere of radius 30, and will find anything that has Object Type to be Pawn OR PhysicsObject. 
+-   Block: The query returns information about the collision in the Out Hit output, and it stops looking for any other collisions.
 
 [VAR:ToolTipFullLink]
-Engine/Physics/Collision
+collision-in-unreal-engine---overview
 [/VAR]
  
 [/EXCERPT:TraceChannel]
 
 
 [EXCERPT:Preset]
-We provide presets of collision settings for components to use. 
 
-![](CollisionPreset.png)
+Every actor that can collide with anything else has a collision preset: a reusable profile that defines how the actor behaves in the physics collision system. 
 
-If you modify this, note that it can change an existing object's behavior. 
-Please be careful when you modify these values. If you would like to modify this, please go to Project Settings/Collision.
+---
 
-You can also add your presets and assign them to components. 
+Each collision preset contains the following settings:
 
-You need to save this to default config otherwise your change will not show to other users. 
-[/EXCERPT:Preset]
+-   Collision Enabled: Determines which kinds of collision tests take this actor into account. You may be able to increase performance by enabling only the collision tests you need.
 
+-   Object Type: Determines which object channel this actor uses. The engine uses this when testing for collisions between actors, and you can use it when you run your own queries against objects.
 
-[EXCERPT:PresetDetail]
- 
+-   Collision Responses: Determines how this actor responds to trace queries and potential collisions with other actors.
+
+![](CollisionPresetExample.png)
+
+---
 
+Unreal Engine ships with many general-purpose collision presets that emulate common behavior patterns. By creating your own custom collision presets, you can also set up new patterns of physical interaction between different types of objects in your project.
+
+-   To change which preset is assigned to an actor, select the actor in the Level Viewport or Outliner and look for the Collision Presets setting in the Details panel.
+
+-   To create, modify, or delete collision presets, open the Project Settings and look for the Collision section. Note that if you change or delete an existing collision preset, you immediately change the behavior of all actors that are set to use that preset.
 
 [VAR:ToolTipFullLink]
-Engine/Physics/Collision
+collision-in-unreal-engine---overview
 [/VAR]
+[/EXCERPT:Preset]
 
+
+[EXCERPT:PresetDetail]
+[VAR:ExcerptAlias]
+Preset
+[/VAR]
+[VAR:ToolTipFullLink]
+collision-in-unreal-engine---overview
+[/VAR]
 [/EXCERPT:PresetDetail]
+
+
