Documentation for build tools

diff --git a/Engine/Documentation/Source/Programming/UnrealBuildSystem/Configuration/BuildConfigProperties/BuildConfigProperties.INT.udn b/Engine/Documentation/Source/Programming/UnrealBuildSystem/Configuration/BuildConfigProperties/BuildConfigProperties.INT.udn
--- a/Engine/Documentation/Source/Programming/UnrealBuildSystem/Configuration/BuildConfigProperties/BuildConfigProperties.INT.udn
+++ b/Engine/Documentation/Source/Programming/UnrealBuildSystem/Configuration/BuildConfigProperties/BuildConfigProperties.INT.udn
@@ -650,6 +650,8 @@ $ AddDebugCoreConfig : Whether VS Code project generation should include debug c
 
 $ NoCompileCommands : Do not create compile commands json files with compiler arguments for each file; works better with VS Code extension using UBT server mode.
 
+$ UseVSCodeExtension : Create a workspace file for use with VS Code extension that communicates directly with UBT.
+
 $ bIncludeDocumentation : True if we should include documentation in the generated projects.
 
 $ bUsePrecompiled : True if build targets should pass the -useprecompiled argument.
