Documentation for build tools

diff --git a/Engine/Documentation/Source/Programming/UnrealBuildSystem/Configuration/BuildConfigProperties/BuildConfigProperties.INT.udn b/Engine/Documentation/Source/Programming/UnrealBuildSystem/Configuration/BuildConfigProperties/BuildConfigProperties.INT.udn
--- a/Engine/Documentation/Source/Programming/UnrealBuildSystem/Configuration/BuildConfigProperties/BuildConfigProperties.INT.udn
+++ b/Engine/Documentation/Source/Programming/UnrealBuildSystem/Configuration/BuildConfigProperties/BuildConfigProperties.INT.udn
@@ -80,6 +80,12 @@ $ UnsafeTypeCastWarningLevel : Indicates what warning/error level to treat unsaf
 
 $ bUndefinedIdentifierErrors : Forces the use of undefined identifiers in conditional expressions to be treated as errors.
 
+$ ModuleIncludePathWarningLevel : How to treat general module include path validation messages
+
+$ ModuleIncludePrivateWarningLevel : How to treat private module include path validation messages, where a module is adding an include path that exposes private headers
+
+$ ModuleIncludeSubdirectoryWarningLevel : How to treat unnecessary module sub-directory include path validation messages
+
 $ bRetainFramePointers : Forces frame pointers to be retained this is usually required when you want reliable callstacks e.g. mallocframeprofiler
 
 $ bUseFastMonoCalls : New Monolithic Graphics drivers have optional "fast calls" replacing various D3d functions
diff --git a/Engine/Documentation/Source/Programming/UnrealBuildSystem/ModuleFiles/ModuleFilesProperties/ModuleFilesProperties.INT.udn b/Engine/Documentation/Source/Programming/UnrealBuildSystem/ModuleFiles/ModuleFilesProperties/ModuleFilesProperties.INT.udn
--- a/Engine/Documentation/Source/Programming/UnrealBuildSystem/ModuleFiles/ModuleFilesProperties/ModuleFilesProperties.INT.udn
+++ b/Engine/Documentation/Source/Programming/UnrealBuildSystem/ModuleFiles/ModuleFilesProperties/ModuleFilesProperties.INT.udn
@@ -188,6 +188,12 @@ $ ShadowVariableWarningLevel (WarningLevel): How to treat shadow variable warnin
 
 $ UnsafeTypeCastWarningLevel (WarningLevel): How to treat unsafe implicit type cast warnings (e.g., double->float or int64->int32)
 
+$ ModuleIncludePathWarningLevel (WarningLevel): How to treat general module include path validation messages
+
+$ ModuleIncludePrivateWarningLevel (WarningLevel): How to treat private module include path validation messages, where a module is adding an include path that exposes private headers
+
+$ ModuleIncludeSubdirectoryWarningLevel (WarningLevel): How to treat unnecessary module sub-directory include path validation messages
+
 $ bUseUnity (Boolean): If unity builds are enabled this can be used to override if this specific module will build using Unity. This is set using the per module configurations in BuildConfiguration.
 
 $ bEnforceIWYU (Boolean): Enforce "include what you use" rules when PCHUsage is set to ExplicitOrSharedPCH; warns when monolithic headers (Engine.h, UnrealEd.h, etc...) are used, and checks that source files include their matching header first.
diff --git a/Engine/Documentation/Source/Programming/UnrealBuildSystem/TargetFiles/TargetFilesProperties/TargetFilesProperties.INT.udn b/Engine/Documentation/Source/Programming/UnrealBuildSystem/TargetFiles/TargetFilesProperties/TargetFilesProperties.INT.udn
--- a/Engine/Documentation/Source/Programming/UnrealBuildSystem/TargetFiles/TargetFilesProperties/TargetFilesProperties.INT.udn
+++ b/Engine/Documentation/Source/Programming/UnrealBuildSystem/TargetFiles/TargetFilesProperties/TargetFilesProperties.INT.udn
@@ -228,6 +228,12 @@ $ UnsafeTypeCastWarningLevel (WarningLevel): Indicates what warning/error level
 
 $ bUndefinedIdentifierErrors (Boolean): Forces the use of undefined identifiers in conditional expressions to be treated as errors.
 
+$ ModuleIncludePathWarningLevel (WarningLevel): How to treat general module include path validation messages
+
+$ ModuleIncludePrivateWarningLevel (WarningLevel): How to treat private module include path validation messages, where a module is adding an include path that exposes private headers
+
+$ ModuleIncludeSubdirectoryWarningLevel (WarningLevel): How to treat unnecessary module sub-directory include path validation messages
+
 $ bRetainFramePointersOverride (Nullable&lt;Boolean&gt;): Forces frame pointers to be retained this is usually required when you want reliable callstacks e.g. mallocframeprofiler
 
 $ bUseFastMonoCalls (Boolean): New Monolithic Graphics drivers have optional "fast calls" replacing various D3d functions
