Documentation for build tools

diff --git a/Engine/Documentation/Source/Programming/UnrealBuildSystem/Configuration/BuildConfigProperties/BuildConfigProperties.INT.udn b/Engine/Documentation/Source/Programming/UnrealBuildSystem/Configuration/BuildConfigProperties/BuildConfigProperties.INT.udn
--- a/Engine/Documentation/Source/Programming/UnrealBuildSystem/Configuration/BuildConfigProperties/BuildConfigProperties.INT.udn
+++ b/Engine/Documentation/Source/Programming/UnrealBuildSystem/Configuration/BuildConfigProperties/BuildConfigProperties.INT.udn
@@ -52,8 +52,6 @@ $ bUseUnityBuild : Whether to unify C++ code into larger files for faster compil
 
 $ bForceUnityBuild : Whether to force C++ source files to be combined into larger files for faster compilation.
 
-$ bMergeModuleAndGeneratedUnityFiles : Whether to merge module and generated unity files for faster compilation.
-
 $ bUseAdaptiveUnityBuild : Use a heuristic to determine which files are currently being iterated on and exclude them from unity blobs, result in faster incremental compile times. The current implementation uses the read-only flag to distinguish the working set, assuming that files will be made writable by the source control system if they are being modified. This is true for Perforce, but not for Git.
 
 $ bAdaptiveUnityDisablesOptimizations : Disable optimization for files that are in the adaptive non-unity working set.
@@ -280,12 +278,10 @@ $ bCompilerTrace : Outputs compile timing information so that it can be analyzed
 
 $ bShowIncludes : Print out files that are included by each source file
 
-### TargetRules
-
-$ bDisableInliningGenCpps : Disable supports for inlining gen.cpps
-
 ### ModuleConfiguration
 
+$ DisableMergingModuleAndGeneratedFilesInUnityFiles : List of modules that disables merging module and generated cpp files in the same unity files.
+
 $ DisableUnityBuild : List of modules to disable unity builds for
 
 $ EnableOptimizeCode : List of modules to enable optimizations for
@@ -296,10 +292,6 @@ $ OptimizeForSize : List of modules to optimize for size. This allows the target
 
 $ OptimizeForSizeAndSpeed : List of modules to optimize for size and speed. This allows the target to override module optimization level Note that this may disable PCH usage if a private PCH is not provided
 
-### TestTargetRules
-
-$ bDisableInliningGenCpps : Disable supports for inlining gen.cpps
-
 ### FASTBuild
 
 $ FBuildExecutablePath : Used to specify the location of fbuild.exe if the distributed binary isn't being used
diff --git a/Engine/Documentation/Source/Programming/UnrealBuildSystem/TargetFiles/TargetFilesProperties/TargetFilesProperties.INT.udn b/Engine/Documentation/Source/Programming/UnrealBuildSystem/TargetFiles/TargetFilesProperties/TargetFilesProperties.INT.udn
--- a/Engine/Documentation/Source/Programming/UnrealBuildSystem/TargetFiles/TargetFilesProperties/TargetFilesProperties.INT.udn
+++ b/Engine/Documentation/Source/Programming/UnrealBuildSystem/TargetFiles/TargetFilesProperties/TargetFilesProperties.INT.udn
@@ -204,7 +204,7 @@ $ bDisableSymbolCache (Boolean): True if debug symbols that are cached for some
 
 $ bForceUnityBuild (Boolean): Whether to force C++ source files to be combined into larger files for faster compilation.
 
-$ bMergeModuleAndGeneratedUnityFiles (Boolean): Whether to merge module and generated unity files for faster compilation.
+$ DisableMergingModuleAndGeneratedFilesInUnityFiles (String[]): List of modules that disables merging module and generated cpp files in the same unity files.
 
 $ bUseAdaptiveUnityBuild (Boolean): Use a heuristic to determine which files are currently being iterated on and exclude them from unity blobs, result in faster incremental compile times. The current implementation uses the read-only flag to distinguish the working set, assuming that files will be made writable by the source control system if they are being modified. This is true for Perforce, but not for Git.
 
@@ -352,8 +352,6 @@ $ bParseTimingInfoForTracing (Boolean): Whether to parse timing data into a trac
 
 $ bPublicSymbolsByDefault (Boolean): Whether to expose all symbols as public by default on POSIX platforms
 
-$ bDisableInliningGenCpps (Boolean): Disable supports for inlining gen.cpps
-
 $ ToolChainName (String): Allows overriding the toolchain to be created for this target. This must match the name of a class declared in the UnrealBuildTool assembly.
 
 $ bDisableUnverifiedCertificates (Boolean): Whether to allow engine configuration to determine if we can load unverified certificates.
