Documentation for build tools

diff --git a/Engine/Documentation/Source/Programming/Development/BuildGraph/ScriptAnatomy/Tasks/PredefinedTasks/BuildGraphScriptPredefinedTasks.INT.udn b/Engine/Documentation/Source/Programming/Development/BuildGraph/ScriptAnatomy/Tasks/PredefinedTasks/BuildGraphScriptPredefinedTasks.INT.udn
--- a/Engine/Documentation/Source/Programming/Development/BuildGraph/ScriptAnatomy/Tasks/PredefinedTasks/BuildGraphScriptPredefinedTasks.INT.udn
+++ b/Engine/Documentation/Source/Programming/Development/BuildGraph/ScriptAnatomy/Tasks/PredefinedTasks/BuildGraphScriptPredefinedTasks.INT.udn
@@ -385,6 +385,18 @@ Creates a PAK file from a given set of files.
 | Tag          | TagList       | Optional | Tag to be applied to build products of this task.                                                                                 |
 | If           | Condition     | Optional | Whether to execute this task. It is ignored if this condition evaluates to false.                                                 |
 
+### QueryTargets
+
+Runs UBT to query all the targets for a particular project
+
+|                   |               |          |                                                                                                               |
+| ----------------- | ------------- | -------- | ------------------------------------------------------------------------------------------------------------- |
+| ProjectFile       | FileReference | Optional | Path to the project file to query                                                                             |
+| OutputFile        | FileReference | Optional | Path to the output file to receive information about the targets                                              |
+| IncludeAllTargets | Boolean       | Optional | Write out all targets, even if a default is specified in the BuildSettings section of the Default*.ini files. |
+| Tag               | TagList       | Optional | Tag to be applied to build products of this task.                                                             |
+| If                | Condition     | Optional | Whether to execute this task. It is ignored if this condition evaluates to false.                             |
+
 ### RandomData
 
 Creates files containing random data in the specified output directory. Used for generating test data for the temp storage system.
@@ -564,6 +576,7 @@ Extract files from a zip archive.
 | ZipFile        | FileSpec           | Required | Path to the zip file to extract.                                                  |
 | ToDir          | DirectoryReference | Required | Output directory for the extracted files.                                         |
 | UseLegacyUnzip | Boolean            | Optional | Whether or not to use the legacy unzip code.                                      |
+| OverwriteFiles | Boolean            | Optional | Whether or not to overwrite files during unzip.                                   |
 | Tag            | TagList            | Optional | Tag to be applied to the extracted files.                                         |
 | If             | Condition          | Optional | Whether to execute this task. It is ignored if this condition evaluates to false. |
 
diff --git a/Engine/Documentation/Source/Programming/UnrealBuildSystem/Configuration/BuildConfigProperties/BuildConfigProperties.INT.udn b/Engine/Documentation/Source/Programming/UnrealBuildSystem/Configuration/BuildConfigProperties/BuildConfigProperties.INT.udn
--- a/Engine/Documentation/Source/Programming/UnrealBuildSystem/Configuration/BuildConfigProperties/BuildConfigProperties.INT.udn
+++ b/Engine/Documentation/Source/Programming/UnrealBuildSystem/Configuration/BuildConfigProperties/BuildConfigProperties.INT.udn
@@ -12,8 +12,6 @@ $ bPrintDebugInfo : Whether debug info should be written to the console.
 
 $ bAllowHybridExecutor : Whether the hybrid executor will be used (a remote executor and local executor).
 
-$ bAllowHordeCompute : Whether Horde remote compute may be used. Highly experimental, disabled by default.
-
 $ bAllowXGE : Whether XGE may be used.
 
 $ bAllowFASTBuild : Whether FASTBuild may be used.
@@ -52,6 +50,8 @@ $ bUseUnityBuild : Whether to unify C++ code into larger files for faster compil
 
 $ bForceUnityBuild : Whether to force C++ source files to be combined into larger files for faster compilation.
 
+$ bMergeModuleAndGeneratedUnityFiles : Whether to merge module and generated unity files for faster compilation.
+
 $ bUseAdaptiveUnityBuild : Use a heuristic to determine which files are currently being iterated on and exclude them from unity blobs, result in faster incremental compile times. The current implementation uses the read-only flag to distinguish the working set, assuming that files will be made writable by the source control system if they are being modified. This is true for Perforce, but not for Git.
 
 $ bAdaptiveUnityDisablesOptimizations : Disable optimization for files that are in the adaptive non-unity working set.
@@ -292,6 +292,14 @@ $ OptimizeForSize : List of modules to optimize for size. This allows the target
 
 $ OptimizeForSizeAndSpeed : List of modules to optimize for size and speed. This allows the target to override module optimization level Note that this may disable PCH usage if a private PCH is not provided
 
+### TargetRules
+
+$ bDisableInliningGenCpps : Disable supports for inlining gen.cpps
+
+### TestTargetRules
+
+$ bDisableInliningGenCpps : Disable supports for inlining gen.cpps
+
 ### FASTBuild
 
 $ FBuildExecutablePath : Used to specify the location of fbuild.exe if the distributed binary isn't being used
@@ -316,14 +324,6 @@ $ MsvcCRTRedistVersion : Which MSVC CRT Redist version to use
 
 $ CompilerVersion : Which MSVC Compiler version to use
 
-### HordeExecutor
-
-$ NumRemoteParallelProcesses : How many processes that will be executed in parallel remotely
-
-$ RemoteProcessOnly : Only run work remotely, for testing
-
-$ RetryFailedRemote : Retry actions locally when failed remotely.
-
 ### ParallelExecutor
 
 $ MaxProcessorCount : Maximum processor count for local execution.
diff --git a/Engine/Documentation/Source/Programming/UnrealBuildSystem/ModuleFiles/ModuleFilesProperties/ModuleFilesProperties.INT.udn b/Engine/Documentation/Source/Programming/UnrealBuildSystem/ModuleFiles/ModuleFilesProperties/ModuleFilesProperties.INT.udn
--- a/Engine/Documentation/Source/Programming/UnrealBuildSystem/ModuleFiles/ModuleFilesProperties/ModuleFilesProperties.INT.udn
+++ b/Engine/Documentation/Source/Programming/UnrealBuildSystem/ModuleFiles/ModuleFilesProperties/ModuleFilesProperties.INT.udn
@@ -50,6 +50,8 @@ $ bEnableObjCExceptions (Boolean): Enable objective C exception handling
 
 $ bEnableObjCAutomaticReferenceCounting (Boolean): Enable objective C automatic reference counting (ARC) If you set this to true you should not use shared PCHs for this module. The engine won't be extensively using ARC in the short term Not doing this will result in a compile errors because shared PCHs were compiled with different flags than consumer
 
+$ bWarningsAsErrors (Boolean): Whether to enable all warnings as errors. UE enables most warnings as errors already, but disables a few (such as deprecation warnings).
+
 $ bEnableUndefinedIdentifierWarnings (Boolean): Enable warnings for using undefined identifiers in #if expressions
 
 $ bDisableStaticAnalysis (Boolean): Disable all static analysis - clang, msvc, pvs-studio.
diff --git a/Engine/Documentation/Source/Programming/UnrealBuildSystem/TargetFiles/TargetFilesProperties/TargetFilesProperties.INT.udn b/Engine/Documentation/Source/Programming/UnrealBuildSystem/TargetFiles/TargetFilesProperties/TargetFilesProperties.INT.udn
--- a/Engine/Documentation/Source/Programming/UnrealBuildSystem/TargetFiles/TargetFilesProperties/TargetFilesProperties.INT.udn
+++ b/Engine/Documentation/Source/Programming/UnrealBuildSystem/TargetFiles/TargetFilesProperties/TargetFilesProperties.INT.udn
@@ -204,6 +204,8 @@ $ bDisableSymbolCache (Boolean): True if debug symbols that are cached for some
 
 $ bForceUnityBuild (Boolean): Whether to force C++ source files to be combined into larger files for faster compilation.
 
+$ bMergeModuleAndGeneratedUnityFiles (Boolean): Whether to merge module and generated unity files for faster compilation.
+
 $ DisableMergingModuleAndGeneratedFilesInUnityFiles (String[]): List of modules that disables merging module and generated cpp files in the same unity files.
 
 $ bUseAdaptiveUnityBuild (Boolean): Use a heuristic to determine which files are currently being iterated on and exclude them from unity blobs, result in faster incremental compile times. The current implementation uses the read-only flag to distinguish the working set, assuming that files will be made writable by the source control system if they are being modified. This is true for Perforce, but not for Git.
@@ -352,6 +354,8 @@ $ bParseTimingInfoForTracing (Boolean): Whether to parse timing data into a trac
 
 $ bPublicSymbolsByDefault (Boolean): Whether to expose all symbols as public by default on POSIX platforms
 
+$ bDisableInliningGenCpps (Boolean): Disable supports for inlining gen.cpps
+
 $ ToolChainName (String): Allows overriding the toolchain to be created for this target. This must match the name of a class declared in the UnrealBuildTool assembly.
 
 $ bDisableUnverifiedCertificates (Boolean): Whether to allow engine configuration to determine if we can load unverified certificates.
