Documentation for build tools

diff --git a/Engine/Documentation/Source/Programming/Development/BuildGraph/ScriptAnatomy/Tasks/PredefinedTasks/BuildGraphScriptPredefinedTasks.INT.udn b/Engine/Documentation/Source/Programming/Development/BuildGraph/ScriptAnatomy/Tasks/PredefinedTasks/BuildGraphScriptPredefinedTasks.INT.udn
--- a/Engine/Documentation/Source/Programming/Development/BuildGraph/ScriptAnatomy/Tasks/PredefinedTasks/BuildGraphScriptPredefinedTasks.INT.udn
+++ b/Engine/Documentation/Source/Programming/Development/BuildGraph/ScriptAnatomy/Tasks/PredefinedTasks/BuildGraphScriptPredefinedTasks.INT.udn
@@ -162,6 +162,18 @@ Delete a set of files.
 | Verbose                | Boolean   | Optional | Whether or not to use verbose logging.                                                                              |
 | If                     | Condition | Optional | Whether to execute this task. It is ignored if this condition evaluates to false.                                   |
 
+### DeployTool
+
+Deploys a tool update through Horde
+
+|           |           |          |                                                                                                         |
+| --------- | --------- | -------- | ------------------------------------------------------------------------------------------------------- |
+| Id        | String    | Required | Identifier for the tool                                                                                 |
+| Settings  | String    | Required | Settings file to use for the deployment. Should be a JSON file containing server name and access token. |
+| File      | String    | Optional | Zip file containing the data to upload                                                                  |
+| Directory | String    | Optional | Directory to zip and upload for the tool                                                                |
+| If        | Condition | Optional | Whether to execute this task. It is ignored if this condition evaluates to false.                       |
+
 ### Docker
 
 Spawns Docker and waits for it to complete.
diff --git a/Engine/Documentation/Source/Programming/UnrealBuildSystem/Configuration/BuildConfigProperties/BuildConfigProperties.INT.udn b/Engine/Documentation/Source/Programming/UnrealBuildSystem/Configuration/BuildConfigProperties/BuildConfigProperties.INT.udn
--- a/Engine/Documentation/Source/Programming/UnrealBuildSystem/Configuration/BuildConfigProperties/BuildConfigProperties.INT.udn
+++ b/Engine/Documentation/Source/Programming/UnrealBuildSystem/Configuration/BuildConfigProperties/BuildConfigProperties.INT.udn
@@ -220,6 +220,10 @@ $ bUsePerFileIntellisense : Whether to generate per-file intellisense data.
 
 $ bEditorDependsOnShaderCompileWorker : Whether to include a dependency on ShaderCompileWorker when generating project files for the editor.
 
+$ TempDirectory : If set TMP\TEMP will be overidden to this directory, each process will create a unique subdirectory in this folder.
+
+$ bDeleteTempDirectory : If set the application temp directory will be deleted on exit, only when running with a single instance mutex.
+
 ### UEBuildConfiguration
 
 $ bForceHeaderGeneration : If true, force header regeneration. Intended for the build machine.
