Documentation for build tools

diff --git a/Engine/Documentation/Source/Programming/Development/BuildGraph/ScriptAnatomy/Tasks/PredefinedTasks/BuildGraphScriptPredefinedTasks.INT.udn b/Engine/Documentation/Source/Programming/Development/BuildGraph/ScriptAnatomy/Tasks/PredefinedTasks/BuildGraphScriptPredefinedTasks.INT.udn
--- a/Engine/Documentation/Source/Programming/Development/BuildGraph/ScriptAnatomy/Tasks/PredefinedTasks/BuildGraphScriptPredefinedTasks.INT.udn
+++ b/Engine/Documentation/Source/Programming/Development/BuildGraph/ScriptAnatomy/Tasks/PredefinedTasks/BuildGraphScriptPredefinedTasks.INT.udn
@@ -170,6 +170,7 @@ Deploys a tool update through Horde
 | --------- | --------- | -------- | ------------------------------------------------------------------------------------------------------- |
 | Id        | String    | Required | Identifier for the tool                                                                                 |
 | Settings  | String    | Required | Settings file to use for the deployment. Should be a JSON file containing server name and access token. |
+| Version   | String    | Required | Version number for the new tool                                                                         |
 | File      | String    | Optional | Zip file containing the data to upload                                                                  |
 | Directory | String    | Optional | Directory to zip and upload for the tool                                                                |
 | If        | Condition | Optional | Whether to execute this task. It is ignored if this condition evaluates to false.                       |
diff --git a/Engine/Documentation/Source/Programming/UnrealBuildSystem/Configuration/BuildConfigProperties/BuildConfigProperties.INT.udn b/Engine/Documentation/Source/Programming/UnrealBuildSystem/Configuration/BuildConfigProperties/BuildConfigProperties.INT.udn
--- a/Engine/Documentation/Source/Programming/UnrealBuildSystem/Configuration/BuildConfigProperties/BuildConfigProperties.INT.udn
+++ b/Engine/Documentation/Source/Programming/UnrealBuildSystem/Configuration/BuildConfigProperties/BuildConfigProperties.INT.udn
@@ -180,6 +180,10 @@ $ bParseTimingInfoForTracing : Whether to parse timing data into a tracing file
 
 $ bPublicSymbolsByDefault : Whether to expose all symbols as public by default on POSIX platforms
 
+$ MSVCCompileActionWeight : The weight(cpu/memory utilization) of a MSVC compile action
+
+$ ClangCompileActionWeight : The weight(cpu/memory utilization) of a clang compile action
+
 $ CppStandard : Which C++ stanard to use for compiling this target
 
 $ CStandard : Which C standard to use for compiling this target
@@ -208,6 +212,8 @@ $ bEnableMemorySanitizer : Enables memory sanitizer (MSan)
 
 $ bTuneDebugInfoForLLDB : Turns on tuning of debug info for LLDB
 
+$ bDisableDumpSyms : Whether to globally disable calling dump_syms
+
 $ bUseDSYMFiles : Enables the generation of .dsym files. This can be disabled to enable faster iteration times during development.
 
 $ bWriteSolutionOptionFile : Whether to write a solution option (suo) file for the sln.
@@ -254,6 +260,8 @@ $ Compiler : Version of the compiler toolchain to use on Windows platform. A val
 
 $ CompilerVersion : The specific toolchain version to use. This may be a specific version number (for example, "14.13.26128"), the string "Latest" to select the newest available version, or the string "Preview" to select the newest available preview version. By default, and if it is available, we use the toolchain version indicated by WindowsPlatform.DefaultToolChainVersion (otherwise, we use the latest version).
 
+$ bVCFastFail : True if /fastfail should be passed to the msvc compiler and linker
+
 $ bAllowClangLinker : True if we should use the Clang linker (LLD) when we are compiling with Clang, or Intel linker (xilink\xilib) when we are compiling with Intel oneAPI, otherwise we use the MSVC linker.
 
 $ WindowsSdkVersion : The specific Windows SDK version to use. This may be a specific version number (for example, "8.1", "10.0" or "10.0.10150.0"), or the string "Latest", to select the newest available version. By default, and if it is available, we use the Windows SDK version indicated by WindowsPlatform.DefaultWindowsSdkVersion (otherwise, we use the latest version).
@@ -348,6 +356,10 @@ $ bLogActionCommandLines : Whether to log command lines for actions being execut
 
 $ bPrintActionTargetNames : Add target names for each action executed
 
+$ bUseActionWeights : Whether to take into account the Action's weight when determining to do more work or not.
+
+$ bShowCPUUtilization : Whether to show CPU utilization after the work is complete.
+
 ### SNDBS
 
 $ bAllowOverVpn : When set to false, SNDBS will not be enabled when running connected to the coordinator over VPN. Configure VPN-assigned subnets via the VpnSubnets parameter.
diff --git a/Engine/Documentation/Source/Programming/UnrealBuildSystem/ModuleFiles/ModuleFilesProperties/ModuleFilesProperties.INT.udn b/Engine/Documentation/Source/Programming/UnrealBuildSystem/ModuleFiles/ModuleFilesProperties/ModuleFilesProperties.INT.udn
--- a/Engine/Documentation/Source/Programming/UnrealBuildSystem/ModuleFiles/ModuleFilesProperties/ModuleFilesProperties.INT.udn
+++ b/Engine/Documentation/Source/Programming/UnrealBuildSystem/ModuleFiles/ModuleFilesProperties/ModuleFilesProperties.INT.udn
@@ -8,6 +8,8 @@ Version: 5.3
 
 $ IsPlugin (Boolean): True if a Plugin contains this module
 
+$ Logger (ILogger): Accessor for the target logger
+
 $ HasPackageOverride (Boolean): Returns true if there has been an override type specified on this module
 
 $ AutoSdkDirectory (String): The AutoSDK directory for the active host platform
@@ -166,9 +168,9 @@ $ VerseScope (VerseScope): Visibility of Verse code in this module's Source/Vers
 
 $ bValidateCircularDependencies (Boolean): Whether circular dependencies will be validated against the allow list Circular module dependencies result in slower builds. Disabling this option is strongly discouraged. This option is ignored for Engine modules which will always be validated against the allow list.
 
-$ CppStandard (CppStandardVersion): Which stanard to use for compiling this module
+$ CppStandard (Nullable&lt;CppStandardVersion&gt;): Which stanard to use for compiling this module
 
-$ CStandard (CStandardVersion): Which standard to use for compiling this module
+$ CStandard (Nullable&lt;CStandardVersion&gt;): Which standard to use for compiling this module
 
 $ ModuleSymbolVisibility (SymbolVisibility): Control visibility of symbols
 
diff --git a/Engine/Documentation/Source/Programming/UnrealBuildSystem/TargetFiles/TargetFilesProperties/TargetFilesProperties.INT.udn b/Engine/Documentation/Source/Programming/UnrealBuildSystem/TargetFiles/TargetFilesProperties/TargetFilesProperties.INT.udn
--- a/Engine/Documentation/Source/Programming/UnrealBuildSystem/TargetFiles/TargetFilesProperties/TargetFilesProperties.INT.udn
+++ b/Engine/Documentation/Source/Programming/UnrealBuildSystem/TargetFiles/TargetFilesProperties/TargetFilesProperties.INT.udn
@@ -274,6 +274,8 @@ $ bPreprocessOnly (Boolean): Whether to just preprocess source files for this ta
 
 $ bPreprocessDepends (Boolean): Generate dependency files by preprocessing. This is only recommended when distributing builds as it adds additional overhead.
 
+$ bWithAssembly (Boolean): Whether to generate assembly data while compiling this target. Works exclusively on MSVC for now.
+
 $ StaticAnalyzer (StaticAnalyzer): Whether static code analysis should be enabled.
 
 $ StaticAnalyzerOutputType (StaticAnalyzerOutputType): The output type to use for the static analyzer. This is only supported for Clang.
@@ -358,6 +360,10 @@ $ bDisableInliningGenCpps (Boolean): Disable supports for inlining gen.cpps
 
 $ ToolChainName (String): Allows overriding the toolchain to be created for this target. This must match the name of a class declared in the UnrealBuildTool assembly.
 
+$ MSVCCompileActionWeight (Single): The weight(cpu/memory utilization) of a MSVC compile action
+
+$ ClangCompileActionWeight (Single): The weight(cpu/memory utilization) of a clang compile action
+
 $ bDisableUnverifiedCertificates (Boolean): Whether to allow engine configuration to determine if we can load unverified certificates.
 
 $ bAllowGeneratedIniWhenCooked (Boolean): Whether to load generated ini files in cooked build, (GameUserSettings.ini loaded either way)
@@ -414,6 +420,8 @@ $ OptedInModulePlatforms (UnrealTargetPlatform[]): If this is non-null, then any
 
 $ Name (String): The name of this target
 
+$ Logger (ILogger): Logger for output relating to this target. Set before the constructor is run from
+
 $ DefaultBuildSettings (BuildSettingsVersion): Specifies the engine version to maintain backwards-compatible default build settings with (eg. DefaultSettingsVersion.Release_4_23, DefaultSettingsVersion.Release_4_24). Specify DefaultSettingsVersion.Latest to always use defaults for the current engine version, at the risk of introducing build errors while upgrading.
 
 $ IncludeOrderVersion (EngineIncludeOrderVersion): What version of include order to use when compiling this target. Can be overridden via -ForceIncludeOrder on the command line. ModuleRules.IncludeOrderVersion takes precedence.
