Documentation for build tools

diff --git a/Engine/Documentation/Source/Programming/Development/BuildGraph/ScriptAnatomy/Tasks/PredefinedTasks/BuildGraphScriptPredefinedTasks.INT.udn b/Engine/Documentation/Source/Programming/Development/BuildGraph/ScriptAnatomy/Tasks/PredefinedTasks/BuildGraphScriptPredefinedTasks.INT.udn
--- a/Engine/Documentation/Source/Programming/Development/BuildGraph/ScriptAnatomy/Tasks/PredefinedTasks/BuildGraphScriptPredefinedTasks.INT.udn
+++ b/Engine/Documentation/Source/Programming/Development/BuildGraph/ScriptAnatomy/Tasks/PredefinedTasks/BuildGraphScriptPredefinedTasks.INT.udn
@@ -171,6 +171,8 @@ Deploys a tool update through Horde
 | Id        | String    | Required | Identifier for the tool                                                                                 |
 | Settings  | String    | Required | Settings file to use for the deployment. Should be a JSON file containing server name and access token. |
 | Version   | String    | Required | Version number for the new tool                                                                         |
+| Duration  | Integer   | Optional | Duration over which to roll out the tool, in minutes.                                                   |
+| Paused    | Boolean   | Optional | Whether to create the deployment as paused                                                              |
 | File      | String    | Optional | Zip file containing the data to upload                                                                  |
 | Directory | String    | Optional | Directory to zip and upload for the tool                                                                |
 | If        | Condition | Optional | Whether to execute this task. It is ignored if this condition evaluates to false.                       |
