Documentation for build tools

diff --git a/Engine/Documentation/Source/Programming/Development/BuildGraph/ScriptAnatomy/Tasks/PredefinedTasks/BuildGraphScriptPredefinedTasks.INT.udn b/Engine/Documentation/Source/Programming/Development/BuildGraph/ScriptAnatomy/Tasks/PredefinedTasks/BuildGraphScriptPredefinedTasks.INT.udn
--- a/Engine/Documentation/Source/Programming/Development/BuildGraph/ScriptAnatomy/Tasks/PredefinedTasks/BuildGraphScriptPredefinedTasks.INT.udn
+++ b/Engine/Documentation/Source/Programming/Development/BuildGraph/ScriptAnatomy/Tasks/PredefinedTasks/BuildGraphScriptPredefinedTasks.INT.udn
@@ -640,11 +640,12 @@ Writes text to a file.
 
 Compresses files into a zip archive.
 
-|         |                    |          |                                                                                                                                                             |
-| ------- | ------------------ | -------- | ----------------------------------------------------------------------------------------------------------------------------------------------------------- |
-| FromDir | DirectoryReference | Required | The directory to read compressed files from.                                                                                                                |
-| Files   | FileSpec           | Optional | List of file specifications separated by semicolons (for example, *.cpp;Engine/.../*.bat), or the name of a tag set. Relative paths are taken from FromDir. |
-| ZipFile | FileReference      | Required | The zip file to create.                                                                                                                                     |
-| Tag     | TagList            | Optional | Tag to be applied to the created zip file.                                                                                                                  |
-| If      | Condition          | Optional | Whether to execute this task. It is ignored if this condition evaluates to false.                                                                           |
+|                 |                    |          |                                                                                                                                                             |
+| --------------- | ------------------ | -------- | ----------------------------------------------------------------------------------------------------------------------------------------------------------- |
+| FromDir         | DirectoryReference | Required | The directory to read compressed files from.                                                                                                                |
+| Files           | FileSpec           | Optional | List of file specifications separated by semicolons (for example, *.cpp;Engine/.../*.bat), or the name of a tag set. Relative paths are taken from FromDir. |
+| ExecutableFiles | FileSpec           | Optional | List of files that should have an executable bit set.                                                                                                       |
+| ZipFile         | FileReference      | Required | The zip file to create.                                                                                                                                     |
+| Tag             | TagList            | Optional | Tag to be applied to the created zip file.                                                                                                                  |
+| If              | Condition          | Optional | Whether to execute this task. It is ignored if this condition evaluates to false.                                                                           |
 
