Documentation for build tools

diff --git a/Engine/Documentation/Source/Programming/UnrealBuildSystem/Configuration/BuildConfigProperties/BuildConfigProperties.INT.udn b/Engine/Documentation/Source/Programming/UnrealBuildSystem/Configuration/BuildConfigProperties/BuildConfigProperties.INT.udn
--- a/Engine/Documentation/Source/Programming/UnrealBuildSystem/Configuration/BuildConfigProperties/BuildConfigProperties.INT.udn
+++ b/Engine/Documentation/Source/Programming/UnrealBuildSystem/Configuration/BuildConfigProperties/BuildConfigProperties.INT.udn
@@ -106,11 +106,7 @@ $ bUsePDBFiles : Whether PDB files should be used for Visual C++ builds.
 
 $ bUsePCHFiles : Whether PCH files should be used.
 
-$ bDeterministic : Set flags require for deterministic linking (experimental, may not be fully supported). Deterministic compiling is controlled via ModuleRules.
-
-$ bForceDeterministic : Force set flags require for deterministic compiling and linking (experimental, may not be fully supported). This setting is only recommended for testing, instead:
-        * Set bDeterministic on a per module basis in ModuleRules to control deterministic compiling.
-        * Set bDeterministic on a per target basis in TargetRules to control deterministic linking.
+$ bDeterministic : Set flags require for deterministic compiling and linking (experimental).
 
 $ bChainPCHs : Whether PCHs should be chained when compiling with clang.
 
diff --git a/Engine/Documentation/Source/Programming/UnrealBuildSystem/ModuleFiles/ModuleFilesProperties/ModuleFilesProperties.INT.udn b/Engine/Documentation/Source/Programming/UnrealBuildSystem/ModuleFiles/ModuleFilesProperties/ModuleFilesProperties.INT.udn
--- a/Engine/Documentation/Source/Programming/UnrealBuildSystem/ModuleFiles/ModuleFilesProperties/ModuleFilesProperties.INT.udn
+++ b/Engine/Documentation/Source/Programming/UnrealBuildSystem/ModuleFiles/ModuleFilesProperties/ModuleFilesProperties.INT.udn
@@ -38,8 +38,6 @@ $ ShortName (String): Specifies an alternate name for intermediate directories a
 
 $ bTreatAsEngineModule (Boolean): Whether this module should be treated as an engine module (eg. using engine definitions, PCHs, compiled with optimizations enabled in DebugGame configurations, etc...). Initialized to a default based on the rules assembly it was created from.
 
-$ bDeterministic (Boolean): Set flags for determinstic compiles (experimental, may not be fully supported). Deterministic linking is controlled via TargetRules.
-
 $ bUseRTTI (Boolean): Use run time type information
 
 $ bUseAVX (Boolean): Direct the compiler to generate AVX instructions wherever SSE or AVX intrinsics are used, on the platforms that support it. Note that by enabling this you are changing the minspec for the PC platform, and the resultant executable will crash on machines without AVX support.
@@ -188,6 +186,8 @@ $ DefaultBuildSettings (BuildSettingsVersion): Which engine version's build sett
 
 $ IncludeOrderVersion (EngineIncludeOrderVersion): What version of include order to use when compiling this module. Can be overridden via -ForceIncludeOrder on the command line or in a module's rules.
 
+$ DeterministicWarningLevel (WarningLevel): How to treat deterministic warnings (experimental).
+
 $ ShadowVariableWarningLevel (WarningLevel): How to treat shadow variable warnings
 
 $ UnsafeTypeCastWarningLevel (WarningLevel): How to treat unsafe implicit type cast warnings (e.g., double->float or int64->int32)
diff --git a/Engine/Documentation/Source/Programming/UnrealBuildSystem/TargetFiles/TargetFilesProperties/TargetFilesProperties.INT.udn b/Engine/Documentation/Source/Programming/UnrealBuildSystem/TargetFiles/TargetFilesProperties/TargetFilesProperties.INT.udn
--- a/Engine/Documentation/Source/Programming/UnrealBuildSystem/TargetFiles/TargetFilesProperties/TargetFilesProperties.INT.udn
+++ b/Engine/Documentation/Source/Programming/UnrealBuildSystem/TargetFiles/TargetFilesProperties/TargetFilesProperties.INT.udn
@@ -260,11 +260,7 @@ $ bUsePDBFiles (Boolean): Whether PDB files should be used for Visual C++ builds
 
 $ bUsePCHFiles (Boolean): Whether PCH files should be used.
 
-$ bDeterministic (Boolean): Set flags require for deterministic linking (experimental, may not be fully supported). Deterministic compiling is controlled via ModuleRules.
-
-$ bForceDeterministic (Boolean): Force set flags require for deterministic compiling and linking (experimental, may not be fully supported). This setting is only recommended for testing, instead:
-        * Set bDeterministic on a per module basis in ModuleRules to control deterministic compiling.
-        * Set bDeterministic on a per target basis in TargetRules to control deterministic linking.
+$ bDeterministic (Boolean): Set flags require for deterministic compiling and linking (experimental).
 
 $ bChainPCHs (Boolean): Whether PCHs should be chained when compiling with clang.
 
