Documentation for build tools

diff --git a/Engine/Documentation/Source/Programming/Development/BuildGraph/ScriptAnatomy/Tasks/PredefinedTasks/BuildGraphScriptPredefinedTasks.INT.udn b/Engine/Documentation/Source/Programming/Development/BuildGraph/ScriptAnatomy/Tasks/PredefinedTasks/BuildGraphScriptPredefinedTasks.INT.udn
--- a/Engine/Documentation/Source/Programming/Development/BuildGraph/ScriptAnatomy/Tasks/PredefinedTasks/BuildGraphScriptPredefinedTasks.INT.udn
+++ b/Engine/Documentation/Source/Programming/Development/BuildGraph/ScriptAnatomy/Tasks/PredefinedTasks/BuildGraphScriptPredefinedTasks.INT.udn
@@ -263,7 +263,7 @@ Compiles C# project files, and their dependencies.
 
 |           |               |          |                                                                                   |
 | --------- | ------------- | -------- | --------------------------------------------------------------------------------- |
-| Path      | String        | Optional | Filter Path (default is ...)                                                      |
+| Path      | FileSpec      | Optional | List of file specifications separated by semicolon (default is ...)               |
 | Change    | Integer       | Optional | The configuration to compile.                                                     |
 | MinChange | Integer       | Optional | The configuration to compile.                                                     |
 | MaxChange | Integer       | Optional | The configuration to compile.                                                     |
