Documentation for build tools

diff --git a/Engine/Documentation/Source/Programming/UnrealBuildSystem/TargetFiles/TargetFilesProperties/TargetFilesProperties.INT.udn b/Engine/Documentation/Source/Programming/UnrealBuildSystem/TargetFiles/TargetFilesProperties/TargetFilesProperties.INT.udn
--- a/Engine/Documentation/Source/Programming/UnrealBuildSystem/TargetFiles/TargetFilesProperties/TargetFilesProperties.INT.udn
+++ b/Engine/Documentation/Source/Programming/UnrealBuildSystem/TargetFiles/TargetFilesProperties/TargetFilesProperties.INT.udn
@@ -194,6 +194,8 @@ $ bIWYUHeadersOnly (Boolean): Tells "include what you use" to only compile heade
 
 $ bEnforceIWYU (Boolean): Enforce "include what you use" rules; warns if monolithic headers (Engine.h, UnrealEd.h, etc...) are used, and checks that source files include their matching header first.
 
+$ bWarnAboutMonolithicHeadersIncluded (Boolean): Emit a warning when an old-style monolithic header is included while compiling this target.
+
 $ bPrecompile (Boolean): Make static libraries for all engine modules as intermediates for this target.
 
 $ bEnableOSX109Support (Boolean): Whether we should compile with support for OS X 10.9 Mavericks. Used for some tools that we need to be compatible with this version of OS X.
