Documentation for build tools

diff --git a/Engine/Documentation/Source/Programming/UnrealBuildSystem/Configuration/BuildConfigProperties/BuildConfigProperties.INT.udn b/Engine/Documentation/Source/Programming/UnrealBuildSystem/Configuration/BuildConfigProperties/BuildConfigProperties.INT.udn
--- a/Engine/Documentation/Source/Programming/UnrealBuildSystem/Configuration/BuildConfigProperties/BuildConfigProperties.INT.udn
+++ b/Engine/Documentation/Source/Programming/UnrealBuildSystem/Configuration/BuildConfigProperties/BuildConfigProperties.INT.udn
@@ -40,6 +40,8 @@ $ bLegalToDistributeBinary : Whether the output from this target can be publicly
 
 $ bUseVerse : Whether to use the verse script interface.
 
+$ bUseAutoRTFMCompiler : Whether to use the AutoRTFM Clang compiler.
+
 $ bUseInlining : Enable inlining for all modules.
 
 $ bUseDebugLiveCodingConsole : Whether to enable support for live coding
diff --git a/Engine/Documentation/Source/Programming/UnrealBuildSystem/TargetFiles/TargetFilesProperties/TargetFilesProperties.INT.udn b/Engine/Documentation/Source/Programming/UnrealBuildSystem/TargetFiles/TargetFilesProperties/TargetFilesProperties.INT.udn
--- a/Engine/Documentation/Source/Programming/UnrealBuildSystem/TargetFiles/TargetFilesProperties/TargetFilesProperties.INT.udn
+++ b/Engine/Documentation/Source/Programming/UnrealBuildSystem/TargetFiles/TargetFilesProperties/TargetFilesProperties.INT.udn
@@ -88,6 +88,8 @@ $ bEnableMeshEditor (Boolean): Whether to enable the mesh editor.
 
 $ bUseVerse (Boolean): Whether to use the verse script interface.
 
+$ bUseAutoRTFMCompiler (Boolean): Whether to use the AutoRTFM Clang compiler.
+
 $ bCompileChaos (Boolean): Whether to compile the Chaos physics plugin.
 
 $ bUseChaos (Boolean): Whether to use the Chaos physics interface. This overrides the physx flags to disable APEX and NvCloth
