Documentation for build tools

diff --git a/Engine/Documentation/Source/Programming/UnrealBuildSystem/Configuration/BuildConfigProperties/BuildConfigProperties.INT.udn b/Engine/Documentation/Source/Programming/UnrealBuildSystem/Configuration/BuildConfigProperties/BuildConfigProperties.INT.udn
--- a/Engine/Documentation/Source/Programming/UnrealBuildSystem/Configuration/BuildConfigProperties/BuildConfigProperties.INT.udn
+++ b/Engine/Documentation/Source/Programming/UnrealBuildSystem/Configuration/BuildConfigProperties/BuildConfigProperties.INT.udn
@@ -214,6 +214,8 @@ $ bDisableDumpSyms : Whether to globally disable calling dump_syms
 
 $ bUseDSYMFiles : Enables the generation of .dsym files. This can be disabled to enable faster iteration times during development.
 
+$ bWriteSarif : Whether .sarif files containing errors and warnings are written alongside each .obj, if supported
+
 $ bWriteSolutionOptionFile : Whether to write a solution option (suo) file for the sln.
 
 $ bVsConfigFile : Whether to write a .vsconfig file next to the sln to suggest components to install.
