Documentation for build tools

diff --git a/Engine/Documentation/Source/Programming/UnrealBuildSystem/Configuration/BuildConfigProperties/BuildConfigProperties.INT.udn b/Engine/Documentation/Source/Programming/UnrealBuildSystem/Configuration/BuildConfigProperties/BuildConfigProperties.INT.udn
--- a/Engine/Documentation/Source/Programming/UnrealBuildSystem/Configuration/BuildConfigProperties/BuildConfigProperties.INT.udn
+++ b/Engine/Documentation/Source/Programming/UnrealBuildSystem/Configuration/BuildConfigProperties/BuildConfigProperties.INT.udn
@@ -186,6 +186,8 @@ $ CppStandard : Which C++ stanard to use for compiling this target
 
 $ CStandard : Which C standard to use for compiling this target
 
+$ MinCpuArchX64 : Direct the compiler to generate AVX instructions wherever SSE or AVX intrinsics are used, on the x64 platforms that support it. Note that by enabling this you are changing the minspec for the PC platform, and the resultant executable will crash on machines without AVX support.
+
 $ bStopSNDBSCompilationAfterErrors : When enabled, SN-DBS will stop compiling targets after a compile error occurs.  Recommended, as it saves computing resources for others.
 
 $ bXGENoWatchdogThread : Whether to use the no_watchdog_thread option to prevent VS2015 toolchain stalls.
diff --git a/Engine/Documentation/Source/Programming/UnrealBuildSystem/ModuleFiles/ModuleFilesProperties/ModuleFilesProperties.INT.udn b/Engine/Documentation/Source/Programming/UnrealBuildSystem/ModuleFiles/ModuleFilesProperties/ModuleFilesProperties.INT.udn
--- a/Engine/Documentation/Source/Programming/UnrealBuildSystem/ModuleFiles/ModuleFilesProperties/ModuleFilesProperties.INT.udn
+++ b/Engine/Documentation/Source/Programming/UnrealBuildSystem/ModuleFiles/ModuleFilesProperties/ModuleFilesProperties.INT.udn
@@ -22,6 +22,8 @@ $ ModuleDirectory (String): Property for the directory containing this module. U
 
 $ TestsDirectory (String): Returns module's low level tests directory "Tests".
 
+$ IsRHIModule (Boolean): True if this UE module implements an RHI.
+
 $ IsVcPackageSupported (Boolean): Returns if VcPkg is supported for the build configuration.
 
 
@@ -40,7 +42,7 @@ $ bTreatAsEngineModule (Boolean): Whether this module should be treated as an en
 
 $ bUseRTTI (Boolean): Use run time type information
 
-$ bUseAVX (Boolean): Direct the compiler to generate AVX instructions wherever SSE or AVX intrinsics are used, on the platforms that support it. Note that by enabling this you are changing the minspec for the PC platform, and the resultant executable will crash on machines without AVX support.
+$ MinCpuArchX64 (Nullable&lt;MinimumCpuArchitectureX64&gt;): Direct the compiler to generate AVX instructions wherever SSE or AVX intrinsics are used, on the x64 platforms that support it. Note that by enabling this you are changing the minspec for the PC platform, and the resultant executable will crash on machines without AVX support.
 
 $ bEnableBufferSecurityChecks (Boolean): Enable buffer security checks.  This should usually be enabled as it prevents severe security risks.
 
@@ -186,6 +188,8 @@ $ DefaultBuildSettings (BuildSettingsVersion): Which engine version's build sett
 
 $ IncludeOrderVersion (EngineIncludeOrderVersion): What version of include order to use when compiling this module. Can be overridden via -ForceIncludeOrder on the command line or in a module's rules.
 
+$ bUseAVX (Boolean): Obsolete: Direct the compiler to generate AVX instructions wherever SSE or AVX intrinsics are used, on the platforms that support it. Note that by enabling this you are changing the minspec for the PC platform, and the resultant executable will crash on machines without AVX support.
+
 $ DeterministicWarningLevel (WarningLevel): How to treat deterministic warnings (experimental).
 
 $ ShadowVariableWarningLevel (WarningLevel): How to treat shadow variable warnings
diff --git a/Engine/Documentation/Source/Programming/UnrealBuildSystem/TargetFiles/TargetFilesProperties/TargetFilesProperties.INT.udn b/Engine/Documentation/Source/Programming/UnrealBuildSystem/TargetFiles/TargetFilesProperties/TargetFilesProperties.INT.udn
--- a/Engine/Documentation/Source/Programming/UnrealBuildSystem/TargetFiles/TargetFilesProperties/TargetFilesProperties.INT.udn
+++ b/Engine/Documentation/Source/Programming/UnrealBuildSystem/TargetFiles/TargetFilesProperties/TargetFilesProperties.INT.udn
@@ -374,6 +374,8 @@ $ CppStandard (CppStandardVersion): Which C++ stanard to use for compiling this
 
 $ CStandard (CStandardVersion): Which C standard to use for compiling this target
 
+$ MinCpuArchX64 (MinimumCpuArchitectureX64): Direct the compiler to generate AVX instructions wherever SSE or AVX intrinsics are used, on the x64 platforms that support it. Note that by enabling this you are changing the minspec for the PC platform, and the resultant executable will crash on machines without AVX support.
+
 $ BuildVersion (String): The build version string
 
 $ GlobalDefinitions (List&lt;String&gt;): Macros to define globally across the whole target.
