Documentation for build tools

diff --git a/Engine/Documentation/Source/Programming/UnrealBuildSystem/ModuleFiles/ModuleFilesProperties/ModuleFilesProperties.INT.udn b/Engine/Documentation/Source/Programming/UnrealBuildSystem/ModuleFiles/ModuleFilesProperties/ModuleFilesProperties.INT.udn
--- a/Engine/Documentation/Source/Programming/UnrealBuildSystem/ModuleFiles/ModuleFilesProperties/ModuleFilesProperties.INT.udn
+++ b/Engine/Documentation/Source/Programming/UnrealBuildSystem/ModuleFiles/ModuleFilesProperties/ModuleFilesProperties.INT.udn
@@ -22,8 +22,6 @@ $ ModuleDirectory (String): Property for the directory containing this module. U
 
 $ TestsDirectory (String): Returns module's low level tests directory "Tests".
 
-$ IsRHIModule (Boolean): True if this UE module implements an RHI.
-
 $ IsVcPackageSupported (Boolean): Returns if VcPkg is supported for the build configuration.
 
 
