Documentation for build tools

diff --git a/Engine/Documentation/Source/Programming/UnrealBuildSystem/ModuleFiles/ModuleFilesProperties/ModuleFilesProperties.INT.udn b/Engine/Documentation/Source/Programming/UnrealBuildSystem/ModuleFiles/ModuleFilesProperties/ModuleFilesProperties.INT.udn
--- a/Engine/Documentation/Source/Programming/UnrealBuildSystem/ModuleFiles/ModuleFilesProperties/ModuleFilesProperties.INT.udn
+++ b/Engine/Documentation/Source/Programming/UnrealBuildSystem/ModuleFiles/ModuleFilesProperties/ModuleFilesProperties.INT.udn
@@ -170,6 +170,8 @@ $ CppStandard (Nullable&lt;CppStandardVersion&gt;): Which stanard to use for com
 
 $ CStandard (Nullable&lt;CStandardVersion&gt;): Which standard to use for compiling this module
 
+$ GenerateHeaderFuncs (List&lt;Tuple&lt;String, Action&lt;ILogger, DirectoryReference&gt;&gt;&gt;): A list of subdirectory names and functions that are invoked to generate header files. The subdirectory name is appended to the generated code directory to form a new directory that headers are generated inside.
+
 $ ModuleSymbolVisibility (SymbolVisibility): Control visibility of symbols
 
 $ Name (String): Name of this module
