Documentation for build tools

diff --git a/Engine/Documentation/Source/Programming/Development/BuildGraph/ScriptAnatomy/Tasks/PredefinedTasks/BuildGraphScriptPredefinedTasks.INT.udn b/Engine/Documentation/Source/Programming/Development/BuildGraph/ScriptAnatomy/Tasks/PredefinedTasks/BuildGraphScriptPredefinedTasks.INT.udn
--- a/Engine/Documentation/Source/Programming/Development/BuildGraph/ScriptAnatomy/Tasks/PredefinedTasks/BuildGraphScriptPredefinedTasks.INT.udn
+++ b/Engine/Documentation/Source/Programming/Development/BuildGraph/ScriptAnatomy/Tasks/PredefinedTasks/BuildGraphScriptPredefinedTasks.INT.udn
@@ -383,6 +383,17 @@ Executes MsBuild
 | Verbosity     | String    | Optional | The MSBuild output verbosity.                                                                  |
 | If            | Condition | Optional | Whether to execute this task. It is ignored if this condition evaluates to false.              |
 
+### NuGet-LicenseCheck
+
+Spawns Docker and waits for it to complete.
+
+|                |                    |          |                                                                                                                                                          |
+| -------------- | ------------------ | -------- | -------------------------------------------------------------------------------------------------------------------------------------------------------- |
+| BaseDir        | String             | Required | Base directory for running the command                                                                                                                   |
+| IgnorePackages | String             | Optional | Specifies a list of packages to ignore for version checks, separated by semicolons. Optional version number may be specified with 'name@version' syntax. |
+| LicenseDir     | DirectoryReference | Optional | Directory containing allowed licenses                                                                                                                    |
+| If             | Condition          | Optional | Whether to execute this task. It is ignored if this condition evaluates to false.                                                                        |
+
 ### PakFile
 
 Creates a PAK file from a given set of files.
