Documentation for build tools

diff --git a/Engine/Documentation/Source/Programming/Development/BuildGraph/ScriptAnatomy/Tasks/PredefinedTasks/BuildGraphScriptPredefinedTasks.INT.udn b/Engine/Documentation/Source/Programming/Development/BuildGraph/ScriptAnatomy/Tasks/PredefinedTasks/BuildGraphScriptPredefinedTasks.INT.udn
--- a/Engine/Documentation/Source/Programming/Development/BuildGraph/ScriptAnatomy/Tasks/PredefinedTasks/BuildGraphScriptPredefinedTasks.INT.udn
+++ b/Engine/Documentation/Source/Programming/Development/BuildGraph/ScriptAnatomy/Tasks/PredefinedTasks/BuildGraphScriptPredefinedTasks.INT.udn
@@ -312,6 +312,18 @@ Spawns Helm and waits for it to complete.
 | WorkingDir      | String    | Optional | Base directory for running the command                                            |
 | If              | Condition | Optional | Whether to execute this task. It is ignored if this condition evaluates to false. |
 
+### Horde-CreateReport
+
+Creates a Horde report file, which will be displayed on the dashboard with any job running this task.
+
+|           |           |          |                                                                                   |
+| --------- | --------- | -------- | --------------------------------------------------------------------------------- |
+| Name      | String    | Required | Name for the report                                                               |
+| Scope     | String    | Required | Where to display the report                                                       |
+| Placement | String    | Required | Where to show the report                                                          |
+| Text      | String    | Required | Text to be displayed                                                              |
+| If        | Condition | Optional | Whether to execute this task. It is ignored if this condition evaluates to false. |
+
 ### Kubectl
 
 Spawns Kubectl and waits for it to complete.
diff --git a/Engine/Documentation/Source/Programming/UnrealBuildSystem/Configuration/BuildConfigProperties/BuildConfigProperties.INT.udn b/Engine/Documentation/Source/Programming/UnrealBuildSystem/Configuration/BuildConfigProperties/BuildConfigProperties.INT.udn
--- a/Engine/Documentation/Source/Programming/UnrealBuildSystem/Configuration/BuildConfigProperties/BuildConfigProperties.INT.udn
+++ b/Engine/Documentation/Source/Programming/UnrealBuildSystem/Configuration/BuildConfigProperties/BuildConfigProperties.INT.udn
@@ -292,6 +292,12 @@ $ bCompilerTrace : Outputs compile timing information so that it can be analyzed
 
 $ bShowIncludes : Print out files that are included by each source file
 
+### TargetRules
+
+$ bCompileChaosVisualDebuggerSupport : Whether to compile in Chaos Visual Debugger (CVD) support features to record the state of the physics simulation
+
+$ bDisableInliningGenCpps : Disable supports for inlining gen.cpps
+
 ### ModuleConfiguration
 
 $ DisableMergingModuleAndGeneratedFilesInUnityFiles : List of modules that disables merging module and generated cpp files in the same unity files.
@@ -306,12 +312,10 @@ $ OptimizeForSize : List of modules to optimize for size. This allows the target
 
 $ OptimizeForSizeAndSpeed : List of modules to optimize for size and speed. This allows the target to override module optimization level Note that this may disable PCH usage if a private PCH is not provided
 
-### TargetRules
-
-$ bDisableInliningGenCpps : Disable supports for inlining gen.cpps
-
 ### TestTargetRules
 
+$ bCompileChaosVisualDebuggerSupport : Whether to compile in Chaos Visual Debugger (CVD) support features to record the state of the physics simulation
+
 $ bDisableInliningGenCpps : Disable supports for inlining gen.cpps
 
 ### FASTBuild
diff --git a/Engine/Documentation/Source/Programming/UnrealBuildSystem/TargetFiles/TargetFilesProperties/TargetFilesProperties.INT.udn b/Engine/Documentation/Source/Programming/UnrealBuildSystem/TargetFiles/TargetFilesProperties/TargetFilesProperties.INT.udn
--- a/Engine/Documentation/Source/Programming/UnrealBuildSystem/TargetFiles/TargetFilesProperties/TargetFilesProperties.INT.udn
+++ b/Engine/Documentation/Source/Programming/UnrealBuildSystem/TargetFiles/TargetFilesProperties/TargetFilesProperties.INT.udn
@@ -98,6 +98,8 @@ $ bUseChaosChecked (Boolean): Whether to compile in checked chaos features for d
 
 $ bUseChaosMemoryTracking (Boolean): Whether to compile in chaos memory tracking features
 
+$ bCompileChaosVisualDebuggerSupport (Boolean): Whether to compile in Chaos Visual Debugger (CVD) support features to record the state of the physics simulation
+
 $ bCustomSceneQueryStructure (Boolean): Whether scene query acceleration is done by UE. The physx scene query structure is still created, but we do not use it.
 
 $ bCompilePhysX (Boolean): Whether to include PhysX support.
