Documentation for build tools

diff --git a/Engine/Documentation/Source/Programming/UnrealBuildSystem/Configuration/BuildConfigProperties/BuildConfigProperties.INT.udn b/Engine/Documentation/Source/Programming/UnrealBuildSystem/Configuration/BuildConfigProperties/BuildConfigProperties.INT.udn
--- a/Engine/Documentation/Source/Programming/UnrealBuildSystem/Configuration/BuildConfigProperties/BuildConfigProperties.INT.udn
+++ b/Engine/Documentation/Source/Programming/UnrealBuildSystem/Configuration/BuildConfigProperties/BuildConfigProperties.INT.udn
@@ -204,6 +204,8 @@ $ bSkipClangValidation : Disables clang build verification checks on static libr
 
 $ bEnableAddressSanitizer : Enables address sanitizer (ASan). Only supported for Visual Studio 2019 16.7.0 and up.
 
+$ bEnableLibFuzzer : Enables LibFuzzer. Only supported for Visual Studio 2022 17.0.0 and up.
+
 $ bEnableThreadSanitizer : Enables thread sanitizer (TSan).
 
 $ bEnableUndefinedBehaviorSanitizer : Enables undefined behavior sanitizer (UBSan).
