Documentation for build tools

diff --git a/Engine/Documentation/Source/Programming/UnrealBuildSystem/Configuration/BuildConfigProperties/BuildConfigProperties.INT.udn b/Engine/Documentation/Source/Programming/UnrealBuildSystem/Configuration/BuildConfigProperties/BuildConfigProperties.INT.udn
--- a/Engine/Documentation/Source/Programming/UnrealBuildSystem/Configuration/BuildConfigProperties/BuildConfigProperties.INT.udn
+++ b/Engine/Documentation/Source/Programming/UnrealBuildSystem/Configuration/BuildConfigProperties/BuildConfigProperties.INT.udn
@@ -108,7 +108,7 @@ $ bUsePDBFiles : Whether PDB files should be used for Visual C++ builds.
 
 $ bUsePCHFiles : Whether PCH files should be used.
 
-$ bDeterministic : Set flags require for deterministic compiling and linking (experimental).
+$ bDeterministic : Set flags require for deterministic compiling and linking.
 
 $ bChainPCHs : Whether PCHs should be chained when compiling with clang.
 
diff --git a/Engine/Documentation/Source/Programming/UnrealBuildSystem/TargetFiles/TargetFilesProperties/TargetFilesProperties.INT.udn b/Engine/Documentation/Source/Programming/UnrealBuildSystem/TargetFiles/TargetFilesProperties/TargetFilesProperties.INT.udn
--- a/Engine/Documentation/Source/Programming/UnrealBuildSystem/TargetFiles/TargetFilesProperties/TargetFilesProperties.INT.udn
+++ b/Engine/Documentation/Source/Programming/UnrealBuildSystem/TargetFiles/TargetFilesProperties/TargetFilesProperties.INT.udn
@@ -266,7 +266,7 @@ $ bUsePDBFiles (Boolean): Whether PDB files should be used for Visual C++ builds
 
 $ bUsePCHFiles (Boolean): Whether PCH files should be used.
 
-$ bDeterministic (Boolean): Set flags require for deterministic compiling and linking (experimental).
+$ bDeterministic (Boolean): Set flags require for deterministic compiling and linking.
 
 $ bChainPCHs (Boolean): Whether PCHs should be chained when compiling with clang.
 
