Documentation for build tools

diff --git a/Engine/Documentation/Source/Programming/UnrealBuildSystem/Configuration/BuildConfigProperties/BuildConfigProperties.INT.udn b/Engine/Documentation/Source/Programming/UnrealBuildSystem/Configuration/BuildConfigProperties/BuildConfigProperties.INT.udn
--- a/Engine/Documentation/Source/Programming/UnrealBuildSystem/Configuration/BuildConfigProperties/BuildConfigProperties.INT.udn
+++ b/Engine/Documentation/Source/Programming/UnrealBuildSystem/Configuration/BuildConfigProperties/BuildConfigProperties.INT.udn
@@ -80,6 +80,8 @@ $ UnsafeTypeCastWarningLevel : Indicates what warning/error level to treat unsaf
 
 $ bUndefinedIdentifierErrors : Forces the use of undefined identifiers in conditional expressions to be treated as errors.
 
+$ PCHPerformanceIssueWarningLevel : Indicates what warning/error level to treat potential PCH performance issues.
+
 $ ModuleIncludePathWarningLevel : How to treat general module include path validation messages
 
 $ ModuleIncludePrivateWarningLevel : How to treat private module include path validation messages, where a module is adding an include path that exposes private headers
diff --git a/Engine/Documentation/Source/Programming/UnrealBuildSystem/TargetFiles/TargetFilesProperties/TargetFilesProperties.INT.udn b/Engine/Documentation/Source/Programming/UnrealBuildSystem/TargetFiles/TargetFilesProperties/TargetFilesProperties.INT.udn
--- a/Engine/Documentation/Source/Programming/UnrealBuildSystem/TargetFiles/TargetFilesProperties/TargetFilesProperties.INT.udn
+++ b/Engine/Documentation/Source/Programming/UnrealBuildSystem/TargetFiles/TargetFilesProperties/TargetFilesProperties.INT.udn
@@ -236,6 +236,8 @@ $ UnsafeTypeCastWarningLevel (WarningLevel): Indicates what warning/error level
 
 $ bUndefinedIdentifierErrors (Boolean): Forces the use of undefined identifiers in conditional expressions to be treated as errors.
 
+$ PCHPerformanceIssueWarningLevel (WarningLevel): Indicates what warning/error level to treat potential PCH performance issues.
+
 $ ModuleIncludePathWarningLevel (WarningLevel): How to treat general module include path validation messages
 
 $ ModuleIncludePrivateWarningLevel (WarningLevel): How to treat private module include path validation messages, where a module is adding an include path that exposes private headers
