Documentation for build tools

diff --git a/Engine/Documentation/Source/Programming/UnrealBuildSystem/Configuration/BuildConfigProperties/BuildConfigProperties.INT.udn b/Engine/Documentation/Source/Programming/UnrealBuildSystem/Configuration/BuildConfigProperties/BuildConfigProperties.INT.udn
--- a/Engine/Documentation/Source/Programming/UnrealBuildSystem/Configuration/BuildConfigProperties/BuildConfigProperties.INT.udn
+++ b/Engine/Documentation/Source/Programming/UnrealBuildSystem/Configuration/BuildConfigProperties/BuildConfigProperties.INT.udn
@@ -184,7 +184,9 @@ $ MSVCCompileActionWeight : The weight(cpu/memory utilization) of a MSVC compile
 
 $ ClangCompileActionWeight : The weight(cpu/memory utilization) of a clang compile action
 
-$ CppStandard : Which C++ stanard to use for compiling this target
+$ CppStandard : Which C++ standard to use for compiling this target
+
+$ MaxCppStandard : Which C++ standard is the maximum allowed for compiling modules in this target if they override the Target's CppStandard for any reason.
 
 $ CStandard : Which C standard to use for compiling this target
 
diff --git a/Engine/Documentation/Source/Programming/UnrealBuildSystem/TargetFiles/TargetFilesProperties/TargetFilesProperties.INT.udn b/Engine/Documentation/Source/Programming/UnrealBuildSystem/TargetFiles/TargetFilesProperties/TargetFilesProperties.INT.udn
--- a/Engine/Documentation/Source/Programming/UnrealBuildSystem/TargetFiles/TargetFilesProperties/TargetFilesProperties.INT.udn
+++ b/Engine/Documentation/Source/Programming/UnrealBuildSystem/TargetFiles/TargetFilesProperties/TargetFilesProperties.INT.udn
@@ -374,7 +374,9 @@ $ bAllowGeneratedIniWhenCooked (Boolean): Whether to load generated ini files in
 
 $ bAllowNonUFSIniWhenCooked (Boolean): Whether to load non-ufs ini files in cooked build, (GameUserSettings.ini loaded either way)
 
-$ CppStandard (CppStandardVersion): Which C++ stanard to use for compiling this target
+$ CppStandard (CppStandardVersion): Which C++ standard to use for compiling this target
+
+$ MaxCppStandard (CppStandardVersion): Which C++ standard is the maximum allowed for compiling modules in this target if they override the Target's CppStandard for any reason.
 
 $ CStandard (CStandardVersion): Which C standard to use for compiling this target
 
