Documentation for build tools

diff --git a/Engine/Documentation/Source/Programming/UnrealBuildSystem/Configuration/BuildConfigProperties/BuildConfigProperties.INT.udn b/Engine/Documentation/Source/Programming/UnrealBuildSystem/Configuration/BuildConfigProperties/BuildConfigProperties.INT.udn
--- a/Engine/Documentation/Source/Programming/UnrealBuildSystem/Configuration/BuildConfigProperties/BuildConfigProperties.INT.udn
+++ b/Engine/Documentation/Source/Programming/UnrealBuildSystem/Configuration/BuildConfigProperties/BuildConfigProperties.INT.udn
@@ -274,10 +274,6 @@ $ bAllowClangLinker : True if we should use the Clang linker (LLD) when we are c
 
 $ WindowsSdkVersion : The specific Windows SDK version to use. This may be a specific version number (for example, "8.1", "10.0" or "10.0.10150.0"), or the string "Latest", to select the newest available version. By default, and if it is available, we use the Windows SDK version indicated by WindowsPlatform.DefaultWindowsSdkVersion (otherwise, we use the latest version).
 
-$ StaticAnalyzer : The static analyzer to use.
-
-$ StaticAnalyzerOutputType : The output type to use for the static analyzer.
-
 $ bStrictConformanceMode : Enables strict standard conformance mode (/permissive-).
 
 $ bUpdatedCPPMacro : Enables updated __cplusplus macro (/Zc:__cplusplus).
