Documentation for build tools

diff --git a/Engine/Documentation/Source/Programming/Development/BuildGraph/ScriptAnatomy/Tasks/PredefinedTasks/BuildGraphScriptPredefinedTasks.INT.udn b/Engine/Documentation/Source/Programming/Development/BuildGraph/ScriptAnatomy/Tasks/PredefinedTasks/BuildGraphScriptPredefinedTasks.INT.udn
--- a/Engine/Documentation/Source/Programming/Development/BuildGraph/ScriptAnatomy/Tasks/PredefinedTasks/BuildGraphScriptPredefinedTasks.INT.udn
+++ b/Engine/Documentation/Source/Programming/Development/BuildGraph/ScriptAnatomy/Tasks/PredefinedTasks/BuildGraphScriptPredefinedTasks.INT.udn
@@ -373,14 +373,15 @@ Modifies json files by setting a value specified in the key path
 
 Moves files from one directory to another.
 
-|           |           |          |                                                                                   |
-| --------- | --------- | -------- | --------------------------------------------------------------------------------- |
-| Files     | FileSpec  | Optional | Optional filter to be applied to the list of input files.                         |
-| From      | FileSpec  | Required | The pattern(s) to copy from (for example, Engine/*.txt).                          |
-| To        | FileSpec  | Required | The directory to copy to.                                                         |
-| Overwrite | Boolean   | Optional | Optionally if files should be overwritten, defaults to false.                     |
-| Tag       | TagList   | Optional | Tag to be applied to build products of this task.                                 |
-| If        | Condition | Optional | Whether to execute this task. It is ignored if this condition evaluates to false. |
+|                 |           |          |                                                                                   |
+| --------------- | --------- | -------- | --------------------------------------------------------------------------------- |
+| Files           | FileSpec  | Optional | Optional filter to be applied to the list of input files.                         |
+| From            | FileSpec  | Required | The pattern(s) to copy from (for example, Engine/*.txt).                          |
+| To              | FileSpec  | Required | The directory to copy to.                                                         |
+| Overwrite       | Boolean   | Optional | Optionally if files should be overwritten, defaults to false.                     |
+| Tag             | TagList   | Optional | Tag to be applied to build products of this task.                                 |
+| ErrorIfNotFound | Boolean   | Optional | Whether or not to throw an error if no files were found to copy                   |
+| If              | Condition | Optional | Whether to execute this task. It is ignored if this condition evaluates to false. |
 
 ### MsBuild
 
