Documentation for build tools

diff --git a/Engine/Documentation/Source/Programming/UnrealBuildSystem/Configuration/BuildConfigProperties/BuildConfigProperties.INT.udn b/Engine/Documentation/Source/Programming/UnrealBuildSystem/Configuration/BuildConfigProperties/BuildConfigProperties.INT.udn
--- a/Engine/Documentation/Source/Programming/UnrealBuildSystem/Configuration/BuildConfigProperties/BuildConfigProperties.INT.udn
+++ b/Engine/Documentation/Source/Programming/UnrealBuildSystem/Configuration/BuildConfigProperties/BuildConfigProperties.INT.udn
@@ -34,6 +34,12 @@ $ bAllCores : Consider logical cores when determining how many total cpu cores a
 
 $ bCompactOutput : Instruct the executor to write compact output e.g. only errors, if supported by the executor. This field is used to hold the value when specified from the command line or XML
 
+$ bArtifactRead : If set, artifacts will be read
+
+$ bArtifactWrites : If set, artifacts will be written
+
+$ ArtifactDirectory : Location to store the artifacts.
+
 $ bDebugBuildsActuallyUseDebugCRT : Enables the debug C++ runtime (CRT) for debug builds. By default we always use the release runtime, since the debug version isn't particularly useful when debugging Unreal Engine projects, and linking against the debug CRT libraries forces our third party library dependencies to also be compiled using the debug CRT (and often perform more slowly). Often it can be inconvenient to require a separate copy of the debug versions of third party static libraries simply so that you can debug your program's code.
 
 $ bLegalToDistributeBinary : Whether the output from this target can be publicly distributed, even if it has dependencies on modules that are in folders with special restrictions (eg. CarefullyRedist, NotForLicensees, NoRedist).
