Documentation for build tools

diff --git a/Engine/Documentation/Source/Programming/UnrealBuildSystem/Configuration/BuildConfigProperties/BuildConfigProperties.INT.udn b/Engine/Documentation/Source/Programming/UnrealBuildSystem/Configuration/BuildConfigProperties/BuildConfigProperties.INT.udn
--- a/Engine/Documentation/Source/Programming/UnrealBuildSystem/Configuration/BuildConfigProperties/BuildConfigProperties.INT.udn
+++ b/Engine/Documentation/Source/Programming/UnrealBuildSystem/Configuration/BuildConfigProperties/BuildConfigProperties.INT.udn
@@ -192,8 +192,6 @@ $ ClangCompileActionWeight : The weight(cpu/memory utilization) of a clang compi
 
 $ CppStandard : Which C++ standard to use for compiling this target
 
-$ MaxCppStandard : Which C++ standard is the maximum allowed for compiling modules in this target if they override the Target's CppStandard for any reason.
-
 $ CStandard : Which C standard to use for compiling this target
 
 $ MinCpuArchX64 : Direct the compiler to generate AVX instructions wherever SSE or AVX intrinsics are used, on the x64 platforms that support it. Note that by enabling this you are changing the minspec for the PC platform, and the resultant executable will crash on machines without AVX support.
@@ -814,10 +812,6 @@ $ RsyncAuthentication : The authentication used for Rsync (for the -e rsync flag
 
 $ SshAuthentication : The authentication used for SSH (probably similar to RsyncAuthentication).
 
-### HoloLensPlatform
-
-$ Compiler : Version of the compiler toolchain to use on HoloLens. A value of "default" will be changed to a specific version at UBT startup.
-
 ### Log
 
 $ bBackupLogFiles : Whether to backup an existing log file, rather than overwriting it.
diff --git a/Engine/Documentation/Source/Programming/UnrealBuildSystem/TargetFiles/TargetFilesProperties/TargetFilesProperties.INT.udn b/Engine/Documentation/Source/Programming/UnrealBuildSystem/TargetFiles/TargetFilesProperties/TargetFilesProperties.INT.udn
--- a/Engine/Documentation/Source/Programming/UnrealBuildSystem/TargetFiles/TargetFilesProperties/TargetFilesProperties.INT.udn
+++ b/Engine/Documentation/Source/Programming/UnrealBuildSystem/TargetFiles/TargetFilesProperties/TargetFilesProperties.INT.udn
@@ -376,8 +376,6 @@ $ bAllowNonUFSIniWhenCooked (Boolean): Whether to load non-ufs ini files in cook
 
 $ CppStandard (CppStandardVersion): Which C++ standard to use for compiling this target
 
-$ MaxCppStandard (CppStandardVersion): Which C++ standard is the maximum allowed for compiling modules in this target if they override the Target's CppStandard for any reason.
-
 $ CStandard (CStandardVersion): Which C standard to use for compiling this target
 
 $ MinCpuArchX64 (MinimumCpuArchitectureX64): Direct the compiler to generate AVX instructions wherever SSE or AVX intrinsics are used, on the x64 platforms that support it. Note that by enabling this you are changing the minspec for the PC platform, and the resultant executable will crash on machines without AVX support.
@@ -430,6 +428,10 @@ $ Name (String): The name of this target
 
 $ Logger (ILogger): Logger for output relating to this target. Set before the constructor is run from
 
+$ AdditionalData (Object): Generic nullable object so a user can set additional data in a project's TargetRule and access in a project's ModuleRule without needing to add new properties post-release. For example:
+        * in .Target.cs:	AdditionalData = "data";
+        * in .Build.cs:		if ((Target.AdditionalData as string) == "data") { ... }
+
 $ DefaultBuildSettings (BuildSettingsVersion): Specifies the engine version to maintain backwards-compatible default build settings with (eg. DefaultSettingsVersion.Release_4_23, DefaultSettingsVersion.Release_4_24). Specify DefaultSettingsVersion.Latest to always use defaults for the current engine version, at the risk of introducing build errors while upgrading.
 
 $ IncludeOrderVersion (EngineIncludeOrderVersion): What version of include order to use when compiling this target. Can be overridden via -ForceIncludeOrder on the command line. ModuleRules.IncludeOrderVersion takes precedence.
