Documentation for build tools

diff --git a/Engine/Documentation/Source/Programming/Development/BuildGraph/ScriptAnatomy/Tasks/PredefinedTasks/BuildGraphScriptPredefinedTasks.INT.udn b/Engine/Documentation/Source/Programming/Development/BuildGraph/ScriptAnatomy/Tasks/PredefinedTasks/BuildGraphScriptPredefinedTasks.INT.udn
--- a/Engine/Documentation/Source/Programming/Development/BuildGraph/ScriptAnatomy/Tasks/PredefinedTasks/BuildGraphScriptPredefinedTasks.INT.udn
+++ b/Engine/Documentation/Source/Programming/Development/BuildGraph/ScriptAnatomy/Tasks/PredefinedTasks/BuildGraphScriptPredefinedTasks.INT.udn
@@ -647,6 +647,17 @@ Waits a defined number of seconds.
 | Seconds | Integer   | Required | Number of seconds to wait.                                                        |
 | If      | Condition | Optional | Whether to execute this task. It is ignored if this condition evaluates to false. |
 
+### WriteJsonValue
+
+Modifies json files by setting a value specified in the key path
+
+|       |           |          |                                                                                                                                                                                                                                                                                                                                 |
+| ----- | --------- | -------- | ------------------------------------------------------------------------------------------------------------------------------------------------------------------------------------------------------------------------------------------------------------------------------------------------------------------------------- |
+| File  | FileSpec  | Required | Json file(s) which will be modified                                                                                                                                                                                                                                                                                             |
+| Key   | String    | Required | Json element to set in each file. Syntax for this string is a limited subset of JsonPath notation, and may support object properties and array indices. Any array indices which are omitted or out of range will add a new element to the array (eg. '$.foo.bar[]' will add an element to the 'bar' array in the 'foo' object). |
+| Value | String    | Required | New value to set. May be any value JSON value (string, array, object, number, boolean or null).                                                                                                                                                                                                                                 |
+| If    | Condition | Optional | Whether to execute this task. It is ignored if this condition evaluates to false.                                                                                                                                                                                                                                               |
+
 ### WriteTextFile
 
 Writes text to a file.
diff --git a/Engine/Documentation/Source/Programming/UnrealBuildSystem/Configuration/BuildConfigProperties/BuildConfigProperties.INT.udn b/Engine/Documentation/Source/Programming/UnrealBuildSystem/Configuration/BuildConfigProperties/BuildConfigProperties.INT.udn
--- a/Engine/Documentation/Source/Programming/UnrealBuildSystem/Configuration/BuildConfigProperties/BuildConfigProperties.INT.udn
+++ b/Engine/Documentation/Source/Programming/UnrealBuildSystem/Configuration/BuildConfigProperties/BuildConfigProperties.INT.udn
@@ -38,6 +38,8 @@ $ bArtifactRead : If set, artifacts will be read
 
 $ bArtifactWrites : If set, artifacts will be written
 
+$ bLogArtifactCacheMisses : If true, log all artifact cache misses as informational messages
+
 $ ArtifactDirectory : Location to store the artifacts.
 
 $ bDebugBuildsActuallyUseDebugCRT : Enables the debug C++ runtime (CRT) for debug builds. By default we always use the release runtime, since the debug version isn't particularly useful when debugging Unreal Engine projects, and linking against the debug CRT libraries forces our third party library dependencies to also be compiled using the debug CRT (and often perform more slowly). Often it can be inconvenient to require a separate copy of the debug versions of third party static libraries simply so that you can debug your program's code.
