Documentation for build tools

diff --git a/Engine/Documentation/Source/Programming/UnrealBuildSystem/ModuleFiles/ModuleFilesProperties/ModuleFilesProperties.INT.udn b/Engine/Documentation/Source/Programming/UnrealBuildSystem/ModuleFiles/ModuleFilesProperties/ModuleFilesProperties.INT.udn
--- a/Engine/Documentation/Source/Programming/UnrealBuildSystem/ModuleFiles/ModuleFilesProperties/ModuleFilesProperties.INT.udn
+++ b/Engine/Documentation/Source/Programming/UnrealBuildSystem/ModuleFiles/ModuleFilesProperties/ModuleFilesProperties.INT.udn
@@ -98,6 +98,8 @@ $ bUsePrecompiled (Boolean): Whether this module should use precompiled data. Al
 
 $ bAllowConfidentialPlatformDefines (Boolean): Whether this module can use PLATFORM_XXXX style defines, where XXXX is a confidential platform name. This is used to ensure engine or other shared code does not reveal confidential information inside an #if PLATFORM_XXXX block. Licensee game code may want to allow for them, however.
 
+$ bAllowAutoRTFMInstrumentation (Boolean): Enables AutoRTFM instrumentation to this module only when AutoRTFMCompiler is enabled
+
 $ PublicIncludePathModuleNames (List&lt;String&gt;): List of modules names (no path needed) with header files that our module's public headers needs access to, but we don't need to "import" or link against.
 
 $ PublicDependencyModuleNames (List&lt;String&gt;): List of public dependency module names (no path needed) (automatically does the private/public include). These are modules that are required by our public source files.
