Documentation for build tools

diff --git a/Engine/Documentation/Source/Programming/Development/BuildGraph/ScriptAnatomy/Tasks/PredefinedTasks/BuildGraphScriptPredefinedTasks.INT.udn b/Engine/Documentation/Source/Programming/Development/BuildGraph/ScriptAnatomy/Tasks/PredefinedTasks/BuildGraphScriptPredefinedTasks.INT.udn
--- a/Engine/Documentation/Source/Programming/Development/BuildGraph/ScriptAnatomy/Tasks/PredefinedTasks/BuildGraphScriptPredefinedTasks.INT.udn
+++ b/Engine/Documentation/Source/Programming/Development/BuildGraph/ScriptAnatomy/Tasks/PredefinedTasks/BuildGraphScriptPredefinedTasks.INT.udn
@@ -173,8 +173,7 @@ Deploys a tool update through Horde
 | Version   | String    | Required | Version number for the new tool                                                                         |
 | Duration  | Integer   | Optional | Duration over which to roll out the tool, in minutes.                                                   |
 | Paused    | Boolean   | Optional | Whether to create the deployment as paused                                                              |
-| File      | String    | Optional | Zip file containing the data to upload                                                                  |
-| Directory | String    | Optional | Directory to zip and upload for the tool                                                                |
+| Directory | String    | Required | Directory to upload for the tool                                                                        |
 | If        | Condition | Optional | Whether to execute this task. It is ignored if this condition evaluates to false.                       |
 
 ### Docker
