Documentation for build tools

diff --git a/Engine/Documentation/Source/Programming/UnrealBuildSystem/TargetFiles/TargetFilesProperties/TargetFilesProperties.INT.udn b/Engine/Documentation/Source/Programming/UnrealBuildSystem/TargetFiles/TargetFilesProperties/TargetFilesProperties.INT.udn
--- a/Engine/Documentation/Source/Programming/UnrealBuildSystem/TargetFiles/TargetFilesProperties/TargetFilesProperties.INT.udn
+++ b/Engine/Documentation/Source/Programming/UnrealBuildSystem/TargetFiles/TargetFilesProperties/TargetFilesProperties.INT.udn
@@ -16,6 +16,8 @@ $ ProjectFile (FileReference): Path to the project file for the project containi
 
 $ Version (ReadOnlyBuildVersion): The current build version
 
+$ OriginalName (String): Return the original target name without overrides or adornments
+
 $ IsTestTarget (Boolean): Whether this is a low level tests target.
 
 $ ExplicitTestsTarget (Boolean): Whether this is a test target explicitly defined. Explicitley defined test targets always inherit from TestTargetRules and define their own tests. Implicit test targets are created from existing targets when building with -Mode=Test and they include tests from all dependencies.
