Documentation for build tools

diff --git a/Engine/Documentation/Source/Programming/Development/BuildGraph/ScriptAnatomy/Tasks/PredefinedTasks/BuildGraphScriptPredefinedTasks.INT.udn b/Engine/Documentation/Source/Programming/Development/BuildGraph/ScriptAnatomy/Tasks/PredefinedTasks/BuildGraphScriptPredefinedTasks.INT.udn
--- a/Engine/Documentation/Source/Programming/Development/BuildGraph/ScriptAnatomy/Tasks/PredefinedTasks/BuildGraphScriptPredefinedTasks.INT.udn
+++ b/Engine/Documentation/Source/Programming/Development/BuildGraph/ScriptAnatomy/Tasks/PredefinedTasks/BuildGraphScriptPredefinedTasks.INT.udn
@@ -173,7 +173,8 @@ Deploys a tool update through Horde
 | Version   | String    | Required | Version number for the new tool                                                                         |
 | Duration  | Integer   | Optional | Duration over which to roll out the tool, in minutes.                                                   |
 | Paused    | Boolean   | Optional | Whether to create the deployment as paused                                                              |
-| Directory | String    | Required | Directory to upload for the tool                                                                        |
+| File      | String    | Optional | Zip file containing files to upload                                                                     |
+| Directory | String    | Optional | Directory to upload for the tool                                                                        |
 | If        | Condition | Optional | Whether to execute this task. It is ignored if this condition evaluates to false.                       |
 
 ### Docker
diff --git a/Engine/Documentation/Source/Programming/UnrealBuildSystem/TargetFiles/TargetFilesProperties/TargetFilesProperties.INT.udn b/Engine/Documentation/Source/Programming/UnrealBuildSystem/TargetFiles/TargetFilesProperties/TargetFilesProperties.INT.udn
--- a/Engine/Documentation/Source/Programming/UnrealBuildSystem/TargetFiles/TargetFilesProperties/TargetFilesProperties.INT.udn
+++ b/Engine/Documentation/Source/Programming/UnrealBuildSystem/TargetFiles/TargetFilesProperties/TargetFilesProperties.INT.udn
@@ -70,6 +70,8 @@ $ EnablePlugins (List&lt;String&gt;): Additional plugins that should be included
 
 $ DisablePlugins (List&lt;String&gt;): List of plugins to be disabled for this target. Note that the project file may still reference them, so they should be marked as optional to avoid failing to find them at runtime.
 
+$ DisablePluginsConflictWarningLevel (WarningLevel): How to treat conflicts when a disabled plugin is being enabled by another plugin referencing it
+
 $ InternalPluginDependencies (List&lt;String&gt;): A list of Plugin names that are allowed to exist as dependencies without being defined in the uplugin descriptor
 
 $ PakSigningKeysFile (String): Path to the set of pak signing keys to embed in the executable.
