#UEDOC-16262

Corrected an innacuracy in the Sin and Cosine Material Expression tooltips.

diff --git a/Engine/Documentation/Source/Shared/GraphNodes/Material/Material.INT.udn b/Engine/Documentation/Source/Shared/GraphNodes/Material/Material.INT.udn
--- a/Engine/Documentation/Source/Shared/GraphNodes/Material/Material.INT.udn
+++ b/Engine/Documentation/Source/Shared/GraphNodes/Material/Material.INT.udn
@@ -540,7 +540,7 @@ Engine/Rendering/Materials/ExpressionReference/Math#componentmask
 [/EXCERPT:UMaterialExpressionComponentMask]
 
 [EXCERPT:UMaterialExpressionCosine]
-The Cosine expression outputs the cosine of the value input (in radians). Most commonly, this is used to output a continuous oscillating waveform by connecting a Time expression to its input. The output value will cycle back and forth between -1 and 1.
+The Cosine expression outputs the cosine of the value input, where the input value of 1 corresponds to 2*pi radians. Most commonly, this is used to output a continuous oscillating waveform by connecting a Time expression to its input. The output value will cycle back and forth between -1 and 1.
 
 [VAR:ToolTipFullLink]
 Engine/Rendering/Materials/ExpressionReference/Math#cosine
@@ -658,7 +658,7 @@ Engine/Rendering/Materials/ExpressionReference/Math#logarithm2
 [/EXCERPT:UMaterialExpressionLogarithm2]
 
 [EXCERPT:UMaterialExpressionSine]
-The Sine expression outputs the sine of the value input (in radians). Most commonly, this is used to output a continuous oscillating waveform by connecting a Time expression to its input. The output value will cycle back and forth between -1 and 1. The difference between this and the output of the Cosine expression is the output waveform is offset half the period.
+The Sine expression outputs the sine of the value input, where the input value of 1 corresponds to 2*pi radians. Most commonly, this is used to output a continuous oscillating waveform by connecting a Time expression to its input. The output value will cycle back and forth between -1 and 1. The difference between this and the output of the Cosine expression is the output waveform is offset half the period.
 
 [VAR:ToolTipFullLink]
 Engine/Rendering/Materials/ExpressionReference/Math#sine
