Documentation for build tools

diff --git a/Engine/Documentation/Source/Programming/Development/BuildGraph/ScriptAnatomy/Tasks/PredefinedTasks/BuildGraphScriptPredefinedTasks.INT.udn b/Engine/Documentation/Source/Programming/Development/BuildGraph/ScriptAnatomy/Tasks/PredefinedTasks/BuildGraphScriptPredefinedTasks.INT.udn
--- a/Engine/Documentation/Source/Programming/Development/BuildGraph/ScriptAnatomy/Tasks/PredefinedTasks/BuildGraphScriptPredefinedTasks.INT.udn
+++ b/Engine/Documentation/Source/Programming/Development/BuildGraph/ScriptAnatomy/Tasks/PredefinedTasks/BuildGraphScriptPredefinedTasks.INT.udn
@@ -407,6 +407,15 @@ Spawns Docker and waits for it to complete.
 | LicenseDir     | DirectoryReference | Optional | Directory containing allowed licenses                                                                                                                    |
 | If             | Condition          | Optional | Whether to execute this task. It is ignored if this condition evaluates to false.                                                                        |
 
+### OnExit
+
+Spawns an external executable and waits for it to complete.
+
+|         |           |          |                                                                                   |
+| ------- | --------- | -------- | --------------------------------------------------------------------------------- |
+| Command | String    | Required | Executable to spawn.                                                              |
+| If      | Condition | Optional | Whether to execute this task. It is ignored if this condition evaluates to false. |
+
 ### PakFile
 
 Creates a PAK file from a given set of files.
