Documentation for build tools

diff --git a/Engine/Documentation/Source/Programming/Development/BuildGraph/ScriptAnatomy/Tasks/PredefinedTasks/BuildGraphScriptPredefinedTasks.INT.udn b/Engine/Documentation/Source/Programming/Development/BuildGraph/ScriptAnatomy/Tasks/PredefinedTasks/BuildGraphScriptPredefinedTasks.INT.udn
--- a/Engine/Documentation/Source/Programming/Development/BuildGraph/ScriptAnatomy/Tasks/PredefinedTasks/BuildGraphScriptPredefinedTasks.INT.udn
+++ b/Engine/Documentation/Source/Programming/Development/BuildGraph/ScriptAnatomy/Tasks/PredefinedTasks/BuildGraphScriptPredefinedTasks.INT.udn
@@ -490,6 +490,7 @@ Updates the local version files (Engine/Source/Runtime/Launch/Resources/Version.
 | CompatibleChange | Integer   | Optional | The engine compatible changelist to set in the version files.                                                      |
 | Branch           | String    | Required | The branch string.                                                                                                 |
 | Build            | String    | Optional | The build version string.                                                                                          |
+| BuildURL         | String    | Optional | The URL for a running continuous integration job.                                                                  |
 | Licensee         | Boolean   | Optional | Whether to set the IS_LICENSEE_VERSION flag to true.                                                               |
 | Promoted         | Boolean   | Optional | Whether to set the ENGINE_IS_PROMOTED_BUILD flag to true.                                                          |
 | SkipWrite        | Boolean   | Optional | If set, do not write to the files -- just return the version files that would be updated. Useful for local builds. |
diff --git a/Engine/Documentation/Source/Programming/UnrealBuildSystem/Configuration/BuildConfigProperties/BuildConfigProperties.INT.udn b/Engine/Documentation/Source/Programming/UnrealBuildSystem/Configuration/BuildConfigProperties/BuildConfigProperties.INT.udn
--- a/Engine/Documentation/Source/Programming/UnrealBuildSystem/Configuration/BuildConfigProperties/BuildConfigProperties.INT.udn
+++ b/Engine/Documentation/Source/Programming/UnrealBuildSystem/Configuration/BuildConfigProperties/BuildConfigProperties.INT.udn
@@ -150,6 +150,8 @@ $ bPGOProfile : Whether to enable Profile Guided Optimization (PGO) instrumentat
 
 $ bPGOOptimize : Whether to optimize this build with Profile Guided Optimization (PGO).
 
+$ bCodeCoverage : Whether the target requires code coverage compilation and linking.
+
 $ bSupportEditAndContinue : Whether to support edit and continue.
 
 $ bOmitFramePointers : Whether to omit frame pointers or not. Disabling is useful for e.g. memory profiling on the PC.
diff --git a/Engine/Documentation/Source/Programming/UnrealBuildSystem/ModuleFiles/ModuleFilesProperties/ModuleFilesProperties.INT.udn b/Engine/Documentation/Source/Programming/UnrealBuildSystem/ModuleFiles/ModuleFilesProperties/ModuleFilesProperties.INT.udn
--- a/Engine/Documentation/Source/Programming/UnrealBuildSystem/ModuleFiles/ModuleFilesProperties/ModuleFilesProperties.INT.udn
+++ b/Engine/Documentation/Source/Programming/UnrealBuildSystem/ModuleFiles/ModuleFilesProperties/ModuleFilesProperties.INT.udn
@@ -40,6 +40,8 @@ $ bTreatAsEngineModule (Boolean): Whether this module should be treated as an en
 
 $ bUseRTTI (Boolean): Use run time type information
 
+$ bCodeCoverage (Boolean): Enable code coverage compilation/linking support.
+
 $ MinCpuArchX64 (Nullable&lt;MinimumCpuArchitectureX64&gt;): Direct the compiler to generate AVX instructions wherever SSE or AVX intrinsics are used, on the x64 platforms that support it. Note that by enabling this you are changing the minspec for the PC platform, and the resultant executable will crash on machines without AVX support.
 
 $ bEnableBufferSecurityChecks (Boolean): Enable buffer security checks.  This should usually be enabled as it prevents severe security risks.
diff --git a/Engine/Documentation/Source/Programming/UnrealBuildSystem/TargetFiles/TargetFilesProperties/TargetFilesProperties.INT.udn b/Engine/Documentation/Source/Programming/UnrealBuildSystem/TargetFiles/TargetFilesProperties/TargetFilesProperties.INT.udn
--- a/Engine/Documentation/Source/Programming/UnrealBuildSystem/TargetFiles/TargetFilesProperties/TargetFilesProperties.INT.udn
+++ b/Engine/Documentation/Source/Programming/UnrealBuildSystem/TargetFiles/TargetFilesProperties/TargetFilesProperties.INT.udn
@@ -308,6 +308,8 @@ $ bPGOProfile (Boolean): Whether to enable Profile Guided Optimization (PGO) ins
 
 $ bPGOOptimize (Boolean): Whether to optimize this build with Profile Guided Optimization (PGO).
 
+$ bCodeCoverage (Boolean): Whether the target requires code coverage compilation and linking.
+
 $ bSupportEditAndContinue (Boolean): Whether to support edit and continue.
 
 $ bOmitFramePointers (Boolean): Whether to omit frame pointers or not. Disabling is useful for e.g. memory profiling on the PC.
