Documentation for build tools

diff --git a/Engine/Documentation/Source/Programming/Development/BuildGraph/ScriptAnatomy/Tasks/PredefinedTasks/BuildGraphScriptPredefinedTasks.INT.udn b/Engine/Documentation/Source/Programming/Development/BuildGraph/ScriptAnatomy/Tasks/PredefinedTasks/BuildGraphScriptPredefinedTasks.INT.udn
--- a/Engine/Documentation/Source/Programming/Development/BuildGraph/ScriptAnatomy/Tasks/PredefinedTasks/BuildGraphScriptPredefinedTasks.INT.udn
+++ b/Engine/Documentation/Source/Programming/Development/BuildGraph/ScriptAnatomy/Tasks/PredefinedTasks/BuildGraphScriptPredefinedTasks.INT.udn
@@ -681,6 +681,22 @@ Writes text to a file.
 | Tag    | TagList       | Optional | Tag to be applied to build products of this task.                                   |
 | If     | Condition     | Optional | Whether to execute this task. It is ignored if this condition evaluates to false.   |
 
+### ZenExportSnapshot
+
+Exports an snapshot from Zen to a specified destination.
+
+|                            |               |          |                                                                                   |
+| -------------------------- | ------------- | -------- | --------------------------------------------------------------------------------- |
+| Project                    | FileReference | Optional | The project from which to export the snapshot                                     |
+| Platform                   | String        | Optional | The target platform(s) to export the snapshot for                                 |
+| SnapshotDescriptorFile     | FileReference | Optional | A file to create with information about the snapshot that was exported            |
+| DestinationStorageType     | String        | Required | The type of destination to export the snapshot to (cloud, ...)                    |
+| DestinationCloudHost       | String        | Optional | The host name to use when exporting to a cloud destination                        |
+| DestinationCloudNamespace  | String        | Optional | The namespace to use when exporting to a cloud destination                        |
+| DestinationCloudIdentifier | String        | Optional | The identifier to use when exporting to a cloud destination                       |
+| DestinationCloudBucket     | String        | Optional | A custom bucket name to use when exporting to a cloud destination                 |
+| If                         | Condition     | Optional | Whether to execute this task. It is ignored if this condition evaluates to false. |
+
 ### Zip
 
 Compresses files into a zip archive.
diff --git a/Engine/Documentation/Source/Programming/UnrealBuildSystem/Configuration/BuildConfigProperties/BuildConfigProperties.INT.udn b/Engine/Documentation/Source/Programming/UnrealBuildSystem/Configuration/BuildConfigProperties/BuildConfigProperties.INT.udn
--- a/Engine/Documentation/Source/Programming/UnrealBuildSystem/Configuration/BuildConfigProperties/BuildConfigProperties.INT.udn
+++ b/Engine/Documentation/Source/Programming/UnrealBuildSystem/Configuration/BuildConfigProperties/BuildConfigProperties.INT.udn
@@ -194,12 +194,14 @@ $ MSVCCompileActionWeight : The weight(cpu/memory utilization) of a MSVC compile
 
 $ ClangCompileActionWeight : The weight(cpu/memory utilization) of a clang compile action
 
-$ CppStandard : Which C++ standard to use for compiling this target
-
 $ CStandard : Which C standard to use for compiling this target
 
 $ MinCpuArchX64 : Direct the compiler to generate AVX instructions wherever SSE or AVX intrinsics are used, on the x64 platforms that support it. Note that by enabling this you are changing the minspec for the PC platform, and the resultant executable will crash on machines without AVX support.
 
+$ CppStandardEngine : Which C++ standard to use for compiling this target (for engine modules)
+
+$ CppStandard : Which C++ standard to use for compiling this target (for non-engine modules)
+
 $ bStopSNDBSCompilationAfterErrors : When enabled, SN-DBS will stop compiling targets after a compile error occurs.  Recommended, as it saves computing resources for others.
 
 $ bXGENoWatchdogThread : Whether to use the no_watchdog_thread option to prevent VS2015 toolchain stalls.
diff --git a/Engine/Documentation/Source/Programming/UnrealBuildSystem/TargetFiles/TargetFilesProperties/TargetFilesProperties.INT.udn b/Engine/Documentation/Source/Programming/UnrealBuildSystem/TargetFiles/TargetFilesProperties/TargetFilesProperties.INT.udn
--- a/Engine/Documentation/Source/Programming/UnrealBuildSystem/TargetFiles/TargetFilesProperties/TargetFilesProperties.INT.udn
+++ b/Engine/Documentation/Source/Programming/UnrealBuildSystem/TargetFiles/TargetFilesProperties/TargetFilesProperties.INT.udn
@@ -380,8 +380,6 @@ $ bAllowGeneratedIniWhenCooked (Boolean): Whether to load generated ini files in
 
 $ bAllowNonUFSIniWhenCooked (Boolean): Whether to load non-ufs ini files in cooked build, (GameUserSettings.ini loaded either way)
 
-$ CppStandard (CppStandardVersion): Which C++ standard to use for compiling this target
-
 $ CStandard (CStandardVersion): Which C standard to use for compiling this target
 
 $ MinCpuArchX64 (MinimumCpuArchitectureX64): Direct the compiler to generate AVX instructions wherever SSE or AVX intrinsics are used, on the x64 platforms that support it. Note that by enabling this you are changing the minspec for the PC platform, and the resultant executable will crash on machines without AVX support.
@@ -500,6 +498,10 @@ $ bLegacyPublicIncludePaths (Boolean): Add all the public folders as include pat
 
 $ bLegacyParentIncludePaths (Boolean): Add all the parent folders as include paths for the compile environment.
 
+$ CppStandardEngine (CppStandardVersion): Which C++ standard to use for compiling this target (for engine modules)
+
+$ CppStandard (CppStandardVersion): Which C++ standard to use for compiling this target (for non-engine modules)
+
 $ LinkType (TargetLinkType): Specifies how to link modules in this target (monolithic or modular). This is currently protected for backwards compatibility. Call the GetLinkType() accessor until support for the deprecated ShouldCompileMonolithic() override has been removed.
 
 $ LaunchModuleName (String): Specifies the name of the launch module. For modular builds, this is the module that is compiled into the target's executable.
