Documentation for build tools

diff --git a/Engine/Documentation/Source/Programming/Development/BuildGraph/ScriptAnatomy/Tasks/PredefinedTasks/BuildGraphScriptPredefinedTasks.INT.udn b/Engine/Documentation/Source/Programming/Development/BuildGraph/ScriptAnatomy/Tasks/PredefinedTasks/BuildGraphScriptPredefinedTasks.INT.udn
--- a/Engine/Documentation/Source/Programming/Development/BuildGraph/ScriptAnatomy/Tasks/PredefinedTasks/BuildGraphScriptPredefinedTasks.INT.udn
+++ b/Engine/Documentation/Source/Programming/Development/BuildGraph/ScriptAnatomy/Tasks/PredefinedTasks/BuildGraphScriptPredefinedTasks.INT.udn
@@ -414,6 +414,7 @@ Spawns an external executable and waits for it to complete.
 |         |           |          |                                                                                   |
 | ------- | --------- | -------- | --------------------------------------------------------------------------------- |
 | Command | String    | Required | Executable to spawn.                                                              |
+| Lease   | Boolean   | Optional | Whether to execute on lease termination                                           |
 | If      | Condition | Optional | Whether to execute this task. It is ignored if this condition evaluates to false. |
 
 ### PakFile
diff --git a/Engine/Documentation/Source/Programming/UnrealBuildSystem/TargetFiles/TargetFilesProperties/TargetFilesProperties.INT.udn b/Engine/Documentation/Source/Programming/UnrealBuildSystem/TargetFiles/TargetFilesProperties/TargetFilesProperties.INT.udn
--- a/Engine/Documentation/Source/Programming/UnrealBuildSystem/TargetFiles/TargetFilesProperties/TargetFilesProperties.INT.udn
+++ b/Engine/Documentation/Source/Programming/UnrealBuildSystem/TargetFiles/TargetFilesProperties/TargetFilesProperties.INT.udn
@@ -172,6 +172,10 @@ $ bUseLauncherChecks (Boolean): Whether to check that the process was launched t
 
 $ bUseChecksInShipping (Boolean): Whether to turn on checks (asserts) for test/shipping builds.
 
+$ bAllowProfileGPUInTest (Boolean): Whether to turn on GPU markers for Test builds.
+
+$ bAllowProfileGPUInShipping (Boolean): Whether to turn on GPU markers for Shipping builds.
+
 $ bTCHARIsUTF8 (Boolean): Whether to turn on UTF-8 mode, mapping TCHAR to UTF8CHAR.
 
 $ bUseEstimatedUtcNow (Boolean): Whether to use the EstimatedUtcNow or PlatformUtcNow.  EstimatedUtcNow is appropriate in cases where PlatformUtcNow can be slow.
