Documentation for build tools

diff --git a/Engine/Documentation/Source/Programming/UnrealBuildSystem/TargetFiles/TargetFilesProperties/TargetFilesProperties.INT.udn b/Engine/Documentation/Source/Programming/UnrealBuildSystem/TargetFiles/TargetFilesProperties/TargetFilesProperties.INT.udn
--- a/Engine/Documentation/Source/Programming/UnrealBuildSystem/TargetFiles/TargetFilesProperties/TargetFilesProperties.INT.udn
+++ b/Engine/Documentation/Source/Programming/UnrealBuildSystem/TargetFiles/TargetFilesProperties/TargetFilesProperties.INT.udn
@@ -202,7 +202,7 @@ $ bUseXGEController (Boolean): Whether the XGE controller worker and modules sho
 
 $ bIWYU (Boolean): Enables "include what you use" by default for modules in this target. Changes the default PCH mode for any module in this project to PCHUsageMode.UseExplicitOrSharedPCHs.
 
-$ bIWYUHeadersOnly (Boolean): Tells "include what you use" to only compile header files.
+$ bHeadersOnly (Boolean): When used with -IncludeHeaders, only header files will be compiled.
 
 $ bEnforceIWYU (Boolean): Enforce "include what you use" rules; warns if monolithic headers (Engine.h, UnrealEd.h, etc...) are used, and checks that source files include their matching header first.
 
@@ -484,6 +484,8 @@ $ bWithPerfCounters (Boolean): Whether to include PerfCounters support.
 
 $ bWithLiveCoding (Boolean): Whether to enable support for live coding
 
+$ bIncludeHeaders (Boolean): Adds header files in included modules to the build.
+
 $ bHasExports (Boolean): Whether the final executable should export symbols.
 
 $ bBuildAdditionalConsoleApp (Boolean): If true, creates an additional console application. Hack for Windows, where it's not possible to conditionally inherit a parent's console Window depending on how the application is invoked; you have to link the same executable with a different subsystem setting.
