Documentation for build tools

diff --git a/Engine/Documentation/Source/Programming/UnrealBuildSystem/Configuration/BuildConfigProperties/BuildConfigProperties.INT.udn b/Engine/Documentation/Source/Programming/UnrealBuildSystem/Configuration/BuildConfigProperties/BuildConfigProperties.INT.udn
--- a/Engine/Documentation/Source/Programming/UnrealBuildSystem/Configuration/BuildConfigProperties/BuildConfigProperties.INT.udn
+++ b/Engine/Documentation/Source/Programming/UnrealBuildSystem/Configuration/BuildConfigProperties/BuildConfigProperties.INT.udn
@@ -118,8 +118,6 @@ $ bUsePDBFiles : Whether PDB files should be used for Visual C++ builds.
 
 $ bUsePCHFiles : Whether PCH files should be used.
 
-$ bDeterministic : Set flags require for deterministic compiling and linking.
-
 $ bChainPCHs : Whether PCHs should be chained when compiling with clang.
 
 $ bForceIncludePCHHeadersForGenCppFilesWhenPCHIsDisabled : Whether PCH headers should be force included for gen.cpp files when PCH is disabled.
@@ -198,6 +196,8 @@ $ CStandard : Which C standard to use for compiling this target
 
 $ MinCpuArchX64 : Direct the compiler to generate AVX instructions wherever SSE or AVX intrinsics are used, on the x64 platforms that support it. Note that by enabling this you are changing the minspec for the PC platform, and the resultant executable will crash on machines without AVX support.
 
+$ bDeterministic : Set flags require for deterministic compiling and linking. Enabling deterministic mode for msvc disables codegen multithreading so compiling will be slower
+
 $ CppStandardEngine : Which C++ standard to use for compiling this target (for engine modules)
 
 $ CppStandard : Which C++ standard to use for compiling this target (for non-engine modules)
diff --git a/Engine/Documentation/Source/Programming/UnrealBuildSystem/TargetFiles/TargetFilesProperties/TargetFilesProperties.INT.udn b/Engine/Documentation/Source/Programming/UnrealBuildSystem/TargetFiles/TargetFilesProperties/TargetFilesProperties.INT.udn
--- a/Engine/Documentation/Source/Programming/UnrealBuildSystem/TargetFiles/TargetFilesProperties/TargetFilesProperties.INT.udn
+++ b/Engine/Documentation/Source/Programming/UnrealBuildSystem/TargetFiles/TargetFilesProperties/TargetFilesProperties.INT.udn
@@ -276,8 +276,6 @@ $ bUsePDBFiles (Boolean): Whether PDB files should be used for Visual C++ builds
 
 $ bUsePCHFiles (Boolean): Whether PCH files should be used.
 
-$ bDeterministic (Boolean): Set flags require for deterministic compiling and linking.
-
 $ bChainPCHs (Boolean): Whether PCHs should be chained when compiling with clang.
 
 $ bForceIncludePCHHeadersForGenCppFilesWhenPCHIsDisabled (Boolean): Whether PCH headers should be force included for gen.cpp files when PCH is disabled.
@@ -500,6 +498,8 @@ $ DeprecationWarningLevel (WarningLevel): Level to report deprecation warnings a
 
 $ bRetainFramePointers (Boolean): Forces frame pointers to be retained this is usually required when you want reliable callstacks e.g. mallocframeprofiler
 
+$ bDeterministic (Boolean): Set flags require for deterministic compiling and linking. Enabling deterministic mode for msvc disables codegen multithreading so compiling will be slower
+
 $ bLegacyPublicIncludePaths (Boolean): Add all the public folders as include paths for the compile environment.
 
 $ bLegacyParentIncludePaths (Boolean): Add all the parent folders as include paths for the compile environment.
