Documentation for build tools

diff --git a/Engine/Documentation/Source/Programming/UnrealBuildSystem/Configuration/BuildConfigProperties/BuildConfigProperties.INT.udn b/Engine/Documentation/Source/Programming/UnrealBuildSystem/Configuration/BuildConfigProperties/BuildConfigProperties.INT.udn
--- a/Engine/Documentation/Source/Programming/UnrealBuildSystem/Configuration/BuildConfigProperties/BuildConfigProperties.INT.udn
+++ b/Engine/Documentation/Source/Programming/UnrealBuildSystem/Configuration/BuildConfigProperties/BuildConfigProperties.INT.udn
@@ -212,6 +212,8 @@ $ bStopXGECompilationAfterErrors : When enabled, XGE will stop compiling targets
 
 $ BaseLogFileName : Specifies the file to use for logging.
 
+$ IWYUBaseLogFileName : Specifies the file to use for logging.
+
 $ bStripSymbols : Whether to strip iOS symbols or not (implied by Shipping config).
 
 $ bSkipClangValidation : Disables clang build verification checks on static libraries
