Documentation for build tools

diff --git a/Engine/Documentation/Source/Programming/UnrealBuildSystem/Configuration/BuildConfigProperties/BuildConfigProperties.INT.udn b/Engine/Documentation/Source/Programming/UnrealBuildSystem/Configuration/BuildConfigProperties/BuildConfigProperties.INT.udn
--- a/Engine/Documentation/Source/Programming/UnrealBuildSystem/Configuration/BuildConfigProperties/BuildConfigProperties.INT.udn
+++ b/Engine/Documentation/Source/Programming/UnrealBuildSystem/Configuration/BuildConfigProperties/BuildConfigProperties.INT.udn
@@ -106,7 +106,7 @@ $ bDisableModuleNumIncludedBytesPerUnityCPPOverride : Disables overrides that ar
 
 $ bStressTestUnity : Whether to stress test the C++ unity build robustness by including all C++ files files in a project from a single unified file.
 
-$ bDetailedUnityFiles : Whether to add additional information to the unity files, such as '_of_X' in the file name.
+$ bDetailedUnityFiles : Whether to add additional information to the unity files, such as '_of_X' in the file name. Not recommended.
 
 $ bDisableDebugInfo : Whether to globally disable debug info generation; see DebugInfoHeuristics.cs for per-config and per-platform options.
 
@@ -458,10 +458,6 @@ $ bIncludeDevelopmentConfigs : Whether we should include configurations for "Dev
 
 ### IOSToolChain
 
-$ IOSSDKVersion : The version of the iOS SDK to target at build time.
-
-$ BuildIOSVersion : The version of the iOS to allow at build time.
-
 $ bUseDangerouslyFastMode : If this is set, then we do not do any post-compile steps -- except moving the executable into the proper spot on Mac.
 
 ### WindowsTargetRules
diff --git a/Engine/Documentation/Source/Programming/UnrealBuildSystem/TargetFiles/TargetFilesProperties/TargetFilesProperties.INT.udn b/Engine/Documentation/Source/Programming/UnrealBuildSystem/TargetFiles/TargetFilesProperties/TargetFilesProperties.INT.udn
--- a/Engine/Documentation/Source/Programming/UnrealBuildSystem/TargetFiles/TargetFilesProperties/TargetFilesProperties.INT.udn
+++ b/Engine/Documentation/Source/Programming/UnrealBuildSystem/TargetFiles/TargetFilesProperties/TargetFilesProperties.INT.udn
@@ -262,7 +262,7 @@ $ bDisableModuleNumIncludedBytesPerUnityCPPOverride (Boolean): Disables override
 
 $ bStressTestUnity (Boolean): Whether to stress test the C++ unity build robustness by including all C++ files files in a project from a single unified file.
 
-$ bDetailedUnityFiles (Boolean): Whether to add additional information to the unity files, such as '_of_X' in the file name.
+$ bDetailedUnityFiles (Boolean): Whether to add additional information to the unity files, such as '_of_X' in the file name. Not recommended.
 
 $ bForceDebugInfo (Boolean): Whether to force debug info to be generated.
 
