Documentation for build tools

diff --git a/Engine/Documentation/Source/Programming/UnrealBuildSystem/TargetFiles/TargetFilesProperties/TargetFilesProperties.INT.udn b/Engine/Documentation/Source/Programming/UnrealBuildSystem/TargetFiles/TargetFilesProperties/TargetFilesProperties.INT.udn
--- a/Engine/Documentation/Source/Programming/UnrealBuildSystem/TargetFiles/TargetFilesProperties/TargetFilesProperties.INT.udn
+++ b/Engine/Documentation/Source/Programming/UnrealBuildSystem/TargetFiles/TargetFilesProperties/TargetFilesProperties.INT.udn
@@ -26,6 +26,8 @@ $ WithLowLevelTests (Boolean): Controls the value of WITH_LOW_LEVEL_TESTS that d
 
 $ Architecture (UnrealArch): Gets the Architecture in the normal case where there is a single Architecture in Architectures (this will throw an exception if there is more than one architecture specified)
 
+$ bIWYU (Boolean): Enables "include what you use" by default for modules in this target. Changes the default PCH mode for any module in this project to PCHUsageMode.UseExplicitOrSharedPCHs.
+
 $ HostPlatform (UnrealTargetPlatform): Gets the host platform being built on
 
 $ bGenerateProjectFiles (Boolean): Expose the bGenerateProjectFiles flag to targets, so we can modify behavior as appropriate for better intellisense
@@ -36,6 +38,8 @@ $ bIsEngineInstalled (Boolean): Expose a setting for whether or not the engine i
 
 
 ### Read/Write Properties
+$ IntermediateEnvironment (UnrealIntermediateEnvironment): Intermediate environment. Determines if the intermediates end up in a different folder than normal.
+
 $ Type (TargetType): The type of target.
 
 $ ForcedIncludeOrder (Nullable&lt;EngineIncludeOrderVersion&gt;): Force the include order to a specific version. Overrides any Target and Module rules.
@@ -200,8 +204,6 @@ $ NativePointerMemberBehaviorOverride (Nullable&lt;PointerMemberBehavior&gt;): U
 
 $ bUseXGEController (Boolean): Whether the XGE controller worker and modules should be included in the engine build. These are required for distributed shader compilation using the XGE interception interface.
 
-$ bIWYU (Boolean): Enables "include what you use" by default for modules in this target. Changes the default PCH mode for any module in this project to PCHUsageMode.UseExplicitOrSharedPCHs.
-
 $ bHeadersOnly (Boolean): When used with -IncludeHeaders, only header files will be compiled.
 
 $ bEnforceIWYU (Boolean): Enforce "include what you use" rules; warns if monolithic headers (Engine.h, UnrealEd.h, etc...) are used, and checks that source files include their matching header first.
