Documentation for build tools

diff --git a/Engine/Documentation/Source/Programming/UnrealBuildSystem/Configuration/BuildConfigProperties/BuildConfigProperties.INT.udn b/Engine/Documentation/Source/Programming/UnrealBuildSystem/Configuration/BuildConfigProperties/BuildConfigProperties.INT.udn
--- a/Engine/Documentation/Source/Programming/UnrealBuildSystem/Configuration/BuildConfigProperties/BuildConfigProperties.INT.udn
+++ b/Engine/Documentation/Source/Programming/UnrealBuildSystem/Configuration/BuildConfigProperties/BuildConfigProperties.INT.udn
@@ -202,6 +202,8 @@ $ CppStandardEngine : Which C++ standard to use for compiling this target (for e
 
 $ CppStandard : Which C++ standard to use for compiling this target (for non-engine modules)
 
+$ ActionStallReportTime : Number of second of no completed actions to trigger an action stall report. If zero, stall reports will not be enabled.
+
 $ bStopSNDBSCompilationAfterErrors : When enabled, SN-DBS will stop compiling targets after a compile error occurs.  Recommended, as it saves computing resources for others.
 
 $ bXGENoWatchdogThread : Whether to use the no_watchdog_thread option to prevent VS2015 toolchain stalls.
@@ -272,6 +274,8 @@ $ MaxRootPathLength : Maximum recommended root path length.
 
 $ MaxNestedPathLength : Maximum length of a path relative to the root directory. Used on Windows to ensure paths are portable between machines. Defaults to off.
 
+$ bIgnoreStalePGOData : If -PGOOptimize is specified but the linker flags have changed since the last -PGOProfile, this will emit a warning and build without PGO instead of failing during link with LNK1268.
+
 $ bUseFastGenProfile : If specified along with -PGOProfile, then /FASTGENPROFILE will be used instead of /GENPROFILE. This usually means that the PGO data is generated faster, but the resulting data may not yield as efficient optimizations during -PGOOptimize
 
 $ bPGONoExtraCounters : If specified along with -PGOProfile, prevent the usage of extra counters. Please note that by default /FASTGENPROFILE doesnt use extra counters
diff --git a/Engine/Documentation/Source/Programming/UnrealBuildSystem/ModuleFiles/ModuleFilesProperties/ModuleFilesProperties.INT.udn b/Engine/Documentation/Source/Programming/UnrealBuildSystem/ModuleFiles/ModuleFilesProperties/ModuleFilesProperties.INT.udn
--- a/Engine/Documentation/Source/Programming/UnrealBuildSystem/ModuleFiles/ModuleFilesProperties/ModuleFilesProperties.INT.udn
+++ b/Engine/Documentation/Source/Programming/UnrealBuildSystem/ModuleFiles/ModuleFilesProperties/ModuleFilesProperties.INT.udn
@@ -186,6 +186,8 @@ $ OptimizeCode (CodeOptimization): When this module's code should be optimized.
 
 $ OptimizationLevel (OptimizationMode): Allows fine tuning optimization level for speed and\or code size. This requires a private PCH (or NoPCHs, which is not recommended)
 
+$ FPSemantics (FPSemanticsMode): Allows overriding the FP semantics for this module. This requires a private PCH (or NoPCHs, which is not recommended)
+
 $ PCHUsage (PCHUsageMode): Precompiled header usage for this module
 
 $ DefaultBuildSettings (BuildSettingsVersion): Which engine version's build settings to use by default.
diff --git a/Engine/Documentation/Source/Programming/UnrealBuildSystem/TargetFiles/TargetFilesProperties/TargetFilesProperties.INT.udn b/Engine/Documentation/Source/Programming/UnrealBuildSystem/TargetFiles/TargetFilesProperties/TargetFilesProperties.INT.udn
--- a/Engine/Documentation/Source/Programming/UnrealBuildSystem/TargetFiles/TargetFilesProperties/TargetFilesProperties.INT.udn
+++ b/Engine/Documentation/Source/Programming/UnrealBuildSystem/TargetFiles/TargetFilesProperties/TargetFilesProperties.INT.udn
@@ -192,6 +192,8 @@ $ bCompileForSize (Boolean): True if we want to favor optimizing size over speed
 
 $ OptimizationLevel (OptimizationMode): Allows to fine tune optimizations level for speed and\or code size
 
+$ FPSemantics (FPSemanticsMode): Allows setting the FP semantics.
+
 $ bForceCompileDevelopmentAutomationTests (Boolean): Whether to compile development automation tests.
 
 $ bForceCompilePerformanceAutomationTests (Boolean): Whether to compile performance automation tests.
