Update a version number in the documentation about Font DPI settings.

diff --git a/Engine/Documentation/Source/Shared/Enums/EFontDPI/EFontDPI.INT.udn b/Engine/Documentation/Source/Shared/Enums/EFontDPI/EFontDPI.INT.udn
--- a/Engine/Documentation/Source/Shared/Enums/EFontDPI/EFontDPI.INT.udn
+++ b/Engine/Documentation/Source/Shared/Enums/EFontDPI/EFontDPI.INT.udn
@@ -7,7 +7,7 @@ When you create visual designs in an external application and realize them later
 
 You may be able to fix this discrepancy by adjusting the scaling factor in your project’s Font Resolution setting to match the one used by your design application.
 
-- 72 DPI matches common web-based design tools like Figma. It is the default value for UMG in Unreal Engine 5.2 and later.
+- 72 DPI matches common web-based design tools like Figma. It is the default value for UMG in Unreal Engine 5.3 and later.
 
 - 96 DPI matches the scaling factor used internally by Unreal Engine and Slate. It was also the default value for UMG up until Unreal Engine 5.2. Use it if you need to remain consistent with older projects and workflows, or if you need Unreal Engine artists and programmers to refer to font sizes using consistent values.
 
diff --git a/Engine/Documentation/Source/Shared/Types/UUserInterfaceSettings/UUserInterfaceSettings.INT.udn b/Engine/Documentation/Source/Shared/Types/UUserInterfaceSettings/UUserInterfaceSettings.INT.udn
--- a/Engine/Documentation/Source/Shared/Types/UUserInterfaceSettings/UUserInterfaceSettings.INT.udn
+++ b/Engine/Documentation/Source/Shared/Types/UUserInterfaceSettings/UUserInterfaceSettings.INT.udn
@@ -7,7 +7,7 @@ When you create visual designs in an external application and realize them later
 
 You may be able to fix this discrepancy by adjusting the scaling factor in your project’s Font Resolution setting to match the one used by your design application.
 
-- 72 DPI matches common web-based design tools like Figma. It is the default value for UMG in Unreal Engine 5.2 and later.
+- 72 DPI matches common web-based design tools like Figma. It is the default value for UMG in Unreal Engine 5.3 and later.
 
 - 96 DPI matches the scaling factor used internally by Unreal Engine and Slate. It was also the default value for UMG up until Unreal Engine 5.2. Use it if you need to remain consistent with older projects and workflows, or if you need Unreal Engine artists and programmers to refer to font sizes using consistent values.
 
