Documentation for build tools

diff --git a/Engine/Documentation/Source/Programming/UnrealBuildSystem/Configuration/BuildConfigProperties/BuildConfigProperties.INT.udn b/Engine/Documentation/Source/Programming/UnrealBuildSystem/Configuration/BuildConfigProperties/BuildConfigProperties.INT.udn
--- a/Engine/Documentation/Source/Programming/UnrealBuildSystem/Configuration/BuildConfigProperties/BuildConfigProperties.INT.udn
+++ b/Engine/Documentation/Source/Programming/UnrealBuildSystem/Configuration/BuildConfigProperties/BuildConfigProperties.INT.udn
@@ -454,6 +454,8 @@ $ Platforms : Names of platforms to include in the generated project files
 
 $ Configurations : Names of configurations to include in the generated project files. See UnrealTargetConfiguration for valid entries
 
+$ bGatherThirdPartySource : If true, we'll parse subdirectories of third-party projects to locate source and header files to include in the generated projects.  This can make the generated projects quite a bit bigger, but makes it easier to open files directly from the IDE.
+
 $ PrimaryProjectName : Name of the primary project file -- for example, the base file name for the Visual Studio solution file, or the Xcode project file on Mac.
 
 $ bPrimaryProjectNameFromFolder : If true, sets the primary project name according to the name of the folder it is in.
