Documentation for build tools

diff --git a/Engine/Documentation/Source/Programming/UnrealBuildSystem/Configuration/BuildConfigProperties/BuildConfigProperties.INT.udn b/Engine/Documentation/Source/Programming/UnrealBuildSystem/Configuration/BuildConfigProperties/BuildConfigProperties.INT.udn
--- a/Engine/Documentation/Source/Programming/UnrealBuildSystem/Configuration/BuildConfigProperties/BuildConfigProperties.INT.udn
+++ b/Engine/Documentation/Source/Programming/UnrealBuildSystem/Configuration/BuildConfigProperties/BuildConfigProperties.INT.udn
@@ -196,6 +196,8 @@ $ CStandard : Which C standard to use for compiling this target
 
 $ MinCpuArchX64 : Direct the compiler to generate AVX instructions wherever SSE or AVX intrinsics are used, on the x64 platforms that support it. Note that by enabling this you are changing the minspec for the PC platform, and the resultant executable will crash on machines without AVX support.
 
+$ bRequireObjectPtrForAddReferencedObjects : Require TObjectPtr for FReferenceCollector API's. (Needed for compatibility with incremental GC.)
+
 $ bDeterministic : Set flags require for deterministic compiling and linking. Enabling deterministic mode for msvc disables codegen multithreading so compiling will be slower
 
 $ CppStandardEngine : Which C++ standard to use for compiling this target (for engine modules)
diff --git a/Engine/Documentation/Source/Programming/UnrealBuildSystem/TargetFiles/TargetFilesProperties/TargetFilesProperties.INT.udn b/Engine/Documentation/Source/Programming/UnrealBuildSystem/TargetFiles/TargetFilesProperties/TargetFilesProperties.INT.udn
--- a/Engine/Documentation/Source/Programming/UnrealBuildSystem/TargetFiles/TargetFilesProperties/TargetFilesProperties.INT.udn
+++ b/Engine/Documentation/Source/Programming/UnrealBuildSystem/TargetFiles/TargetFilesProperties/TargetFilesProperties.INT.udn
@@ -498,6 +498,8 @@ $ bAdaptiveUnityDisablesPCHForProject (Boolean): Whether to disable force-includ
 
 $ DefaultWarningLevel (WarningLevel): Default treatment of uncategorized warnings
 
+$ bRequireObjectPtrForAddReferencedObjects (Boolean): Require TObjectPtr for FReferenceCollector API's. (Needed for compatibility with incremental GC.)
+
 $ DeprecationWarningLevel (WarningLevel): Level to report deprecation warnings as errors
 
 $ bRetainFramePointers (Boolean): Forces frame pointers to be retained this is usually required when you want reliable callstacks e.g. mallocframeprofiler
