Documentation for build tools

diff --git a/Engine/Documentation/Source/Programming/UnrealBuildSystem/Configuration/BuildConfigProperties/BuildConfigProperties.INT.udn b/Engine/Documentation/Source/Programming/UnrealBuildSystem/Configuration/BuildConfigProperties/BuildConfigProperties.INT.udn
--- a/Engine/Documentation/Source/Programming/UnrealBuildSystem/Configuration/BuildConfigProperties/BuildConfigProperties.INT.udn
+++ b/Engine/Documentation/Source/Programming/UnrealBuildSystem/Configuration/BuildConfigProperties/BuildConfigProperties.INT.udn
@@ -52,6 +52,8 @@ $ bLegalToDistributeBinary : Whether the output from this target can be publicly
 
 $ bUseVerse : Whether to use the verse script interface.
 
+$ bUseVerseBPVM : Whether to use the BPVM to run Verse.
+
 $ bUseAutoRTFMCompiler : Whether to use the AutoRTFM Clang compiler.
 
 $ bUseInlining : Enable inlining for all modules.
diff --git a/Engine/Documentation/Source/Programming/UnrealBuildSystem/TargetFiles/TargetFilesProperties/TargetFilesProperties.INT.udn b/Engine/Documentation/Source/Programming/UnrealBuildSystem/TargetFiles/TargetFilesProperties/TargetFilesProperties.INT.udn
--- a/Engine/Documentation/Source/Programming/UnrealBuildSystem/TargetFiles/TargetFilesProperties/TargetFilesProperties.INT.udn
+++ b/Engine/Documentation/Source/Programming/UnrealBuildSystem/TargetFiles/TargetFilesProperties/TargetFilesProperties.INT.udn
@@ -96,6 +96,8 @@ $ bEnableMeshEditor (Boolean): Whether to enable the mesh editor.
 
 $ bUseVerse (Boolean): Whether to use the verse script interface.
 
+$ bUseVerseBPVM (Boolean): Whether to use the BPVM to run Verse.
+
 $ bUseAutoRTFMCompiler (Boolean): Whether to use the AutoRTFM Clang compiler.
 
 $ bCompileChaos (Boolean): Whether to compile the Chaos physics plugin.
