Documentation for build tools

diff --git a/Engine/Documentation/Source/Programming/UnrealBuildSystem/TargetFiles/TargetFilesProperties/TargetFilesProperties.INT.udn b/Engine/Documentation/Source/Programming/UnrealBuildSystem/TargetFiles/TargetFilesProperties/TargetFilesProperties.INT.udn
--- a/Engine/Documentation/Source/Programming/UnrealBuildSystem/TargetFiles/TargetFilesProperties/TargetFilesProperties.INT.udn
+++ b/Engine/Documentation/Source/Programming/UnrealBuildSystem/TargetFiles/TargetFilesProperties/TargetFilesProperties.INT.udn
@@ -472,6 +472,8 @@ $ bCompileAgainstCoreUObject (Boolean): Enabled for all builds that include the
 
 $ bCompileAgainstApplicationCore (Boolean): Enabled for builds that need to initialize the ApplicationCore module. Command line utilities do not normally need this.
 
+$ bEnableTrace (Boolean): Enabled for builds that want to use against the Trace module for profiling and diagnostics.
+
 $ bCompileAgainstEditor (Boolean): Enabled for editor builds (TargetType.Editor). Can be overridden for programs (TargetType.Program) that would need to compile against editor code. Not available for other target types. Mainly drives the value of WITH_EDITOR.
 
 $ bCompileSpeedTree (Boolean): Whether we should compile in support for SpeedTree or not.
