Documentation for build tools

diff --git a/Engine/Documentation/Source/Programming/UnrealBuildSystem/Configuration/BuildConfigProperties/BuildConfigProperties.INT.udn b/Engine/Documentation/Source/Programming/UnrealBuildSystem/Configuration/BuildConfigProperties/BuildConfigProperties.INT.udn
--- a/Engine/Documentation/Source/Programming/UnrealBuildSystem/Configuration/BuildConfigProperties/BuildConfigProperties.INT.udn
+++ b/Engine/Documentation/Source/Programming/UnrealBuildSystem/Configuration/BuildConfigProperties/BuildConfigProperties.INT.udn
@@ -50,8 +50,6 @@ $ bDebugBuildsActuallyUseDebugCRT : Enables the debug C++ runtime (CRT) for debu
 
 $ bLegalToDistributeBinary : Whether the output from this target can be publicly distributed, even if it has dependencies on modules that are in folders with special restrictions (eg. CarefullyRedist, NotForLicensees, NoRedist).
 
-$ bUseVerse : Whether to use the verse script interface.
-
 $ bUseVerseBPVM : Whether to use the BPVM to run Verse.
 
 $ bUseAutoRTFMCompiler : Whether to use the AutoRTFM Clang compiler.
diff --git a/Engine/Documentation/Source/Programming/UnrealBuildSystem/TargetFiles/TargetFilesProperties/TargetFilesProperties.INT.udn b/Engine/Documentation/Source/Programming/UnrealBuildSystem/TargetFiles/TargetFilesProperties/TargetFilesProperties.INT.udn
--- a/Engine/Documentation/Source/Programming/UnrealBuildSystem/TargetFiles/TargetFilesProperties/TargetFilesProperties.INT.udn
+++ b/Engine/Documentation/Source/Programming/UnrealBuildSystem/TargetFiles/TargetFilesProperties/TargetFilesProperties.INT.udn
@@ -94,8 +94,6 @@ $ GeneratedCodeVersion (EGeneratedCodeVersion): Allow target module to override
 
 $ bEnableMeshEditor (Boolean): Whether to enable the mesh editor.
 
-$ bUseVerse (Boolean): Whether to use the verse script interface.
-
 $ bUseVerseBPVM (Boolean): Whether to use the BPVM to run Verse.
 
 $ bUseAutoRTFMCompiler (Boolean): Whether to use the AutoRTFM Clang compiler.
