Documentation for build tools

diff --git a/Engine/Documentation/Source/Programming/UnrealBuildSystem/Configuration/BuildConfigProperties/BuildConfigProperties.INT.udn b/Engine/Documentation/Source/Programming/UnrealBuildSystem/Configuration/BuildConfigProperties/BuildConfigProperties.INT.udn
--- a/Engine/Documentation/Source/Programming/UnrealBuildSystem/Configuration/BuildConfigProperties/BuildConfigProperties.INT.udn
+++ b/Engine/Documentation/Source/Programming/UnrealBuildSystem/Configuration/BuildConfigProperties/BuildConfigProperties.INT.udn
@@ -302,6 +302,8 @@ $ bStrictPreprocessorConformance : Enables new preprocessor conformance (/Zc:pre
 
 $ bStrictEnumTypesConformance : Enables enum types conformance (/Zc:enumTypes) in VS2022 17.4 Preview 4.0+.
 
+$ bStripPrivateSymbols : Whether to request the linker create a stripped pdb file as part of the build. If enabled the full debug pdb will have the extension .full.pdb
+
 $ PCHMemoryAllocationFactor : Determines the amount of memory that the compiler allocates to construct precompiled headers (/Zm).
 
 $ AdditionalLinkerOptions : Allow the target to specify extra options for linking that aren't otherwise noted here
