Documentation for build tools

diff --git a/Engine/Documentation/Source/Programming/Development/BuildGraph/ScriptAnatomy/Tasks/PredefinedTasks/BuildGraphScriptPredefinedTasks.INT.udn b/Engine/Documentation/Source/Programming/Development/BuildGraph/ScriptAnatomy/Tasks/PredefinedTasks/BuildGraphScriptPredefinedTasks.INT.udn
--- a/Engine/Documentation/Source/Programming/Development/BuildGraph/ScriptAnatomy/Tasks/PredefinedTasks/BuildGraphScriptPredefinedTasks.INT.udn
+++ b/Engine/Documentation/Source/Programming/Development/BuildGraph/ScriptAnatomy/Tasks/PredefinedTasks/BuildGraphScriptPredefinedTasks.INT.udn
@@ -696,8 +696,31 @@ Exports an snapshot from Zen to a specified destination.
 | DestinationCloudNamespace  | String        | Optional | The namespace to use when exporting to a cloud destination                        |
 | DestinationCloudIdentifier | String        | Optional | The identifier to use when exporting to a cloud destination                       |
 | DestinationCloudBucket     | String        | Optional | A custom bucket name to use when exporting to a cloud destination                 |
+| DestinationDir             | String        | Optional | The directory to store the exported oplog                                         |
+| DestinationOplogName       | String        | Optional | The name of the oplog to write on disk                                            |
 | If                         | Condition     | Optional | Whether to execute this task. It is ignored if this condition evaluates to false. |
 
+### ZenImportOplog
+
+Imports an oplog from Zen to a specified destination.
+
+|             |               |          |                                                                                                    |
+| ----------- | ------------- | -------- | -------------------------------------------------------------------------------------------------- |
+| ImportType  | String        | Required | The type of destination to import from to (cloud, file...)                                         |
+| Files       | String        | Optional | comma separated full path to the oplog dir to import into the local zen server Files="Path1,Path2" |
+| Project     | FileReference | Optional | The project from which to import for                                                               |
+| ProjectName | String        | Optional | The name of the newly created Zen Project we will be importing into                                |
+| Platform    | String        | Optional | The target platform to import the snapshot for                                                     |
+| RootDir     | String        | Optional | Root dir for the UE project. Used to derive the Enging folder and the Project folder               |
+| OplogName   | String        | Optional | The name of the imported oplog                                                                     |
+| HostName    | String        | Optional | The host URL for the zen server we are importing from                                              |
+| HostPort    | String        | Optional | The host port for the zen server we are importing from                                             |
+| CloudURL    | String        | Optional | The cloud URL to import from                                                                       |
+| Namespace   | String        | Optional | what namespace to use when importing from cloud                                                    |
+| Bucket      | String        | Optional | what bucket to use when importing from cloud                                                       |
+| Key         | String        | Optional | What key to use when importing from cloud                                                          |
+| If          | Condition     | Optional | Whether to execute this task. It is ignored if this condition evaluates to false.                  |
+
 ### Zip
 
 Compresses files into a zip archive.
