Documentation for build tools

diff --git a/Engine/Documentation/Source/Programming/UnrealBuildSystem/Configuration/BuildConfigProperties/BuildConfigProperties.INT.udn b/Engine/Documentation/Source/Programming/UnrealBuildSystem/Configuration/BuildConfigProperties/BuildConfigProperties.INT.udn
--- a/Engine/Documentation/Source/Programming/UnrealBuildSystem/Configuration/BuildConfigProperties/BuildConfigProperties.INT.udn
+++ b/Engine/Documentation/Source/Programming/UnrealBuildSystem/Configuration/BuildConfigProperties/BuildConfigProperties.INT.udn
@@ -12,11 +12,11 @@ $ bPrintDebugInfo : Whether debug info should be written to the console.
 
 $ bAllowHybridExecutor : Whether the hybrid executor will be used (a remote executor and local executor).
 
-$ bAllowXGE : Whether XGE may be used.
+$ bAllowXGE : Whether XGE may be used if available, default is true.
 
-$ bAllowFASTBuild : Whether FASTBuild may be used.
+$ bAllowFASTBuild : Whether FASTBuild may be used if availabe, default is true.
 
-$ bAllowSNDBS : Whether SN-DBS may be used.
+$ bAllowSNDBS : Whether SN-DBS may be used if availabe, default is true.
 
 $ bUseUBTMakefiles : Enables support for very fast iterative builds by caching target data. Turning this on causes Unreal Build Tool to emit 'UBT Makefiles' for targets when they are built the first time. Subsequent builds will load these Makefiles and begin outdatedness checking and build invocation very quickly. The caveat is that if source files are added or removed to the project, UBT will need to gather information about those in order for your build to complete successfully. Currently, you must run the project file generator after adding/removing source files to tell UBT to re-gather this information. 
     * Events that can invalidate the 'UBT Makefile':
diff --git a/Engine/Documentation/Source/Programming/UnrealBuildSystem/ModuleFiles/ModuleFilesProperties/ModuleFilesProperties.INT.udn b/Engine/Documentation/Source/Programming/UnrealBuildSystem/ModuleFiles/ModuleFilesProperties/ModuleFilesProperties.INT.udn
--- a/Engine/Documentation/Source/Programming/UnrealBuildSystem/ModuleFiles/ModuleFilesProperties/ModuleFilesProperties.INT.udn
+++ b/Engine/Documentation/Source/Programming/UnrealBuildSystem/ModuleFiles/ModuleFilesProperties/ModuleFilesProperties.INT.udn
@@ -190,6 +190,8 @@ $ FPSemantics (FPSemanticsMode): Allows overriding the FP semantics for this mod
 
 $ PCHUsage (PCHUsageMode): Precompiled header usage for this module
 
+$ bValidateFormatStrings (Boolean): Emits compilation errors for incorrect UE_LOG format strings.
+
 $ DefaultBuildSettings (BuildSettingsVersion): Which engine version's build settings to use by default.
 
 $ IncludeOrderVersion (EngineIncludeOrderVersion): What version of include order to use when compiling this module. Can be overridden via -ForceIncludeOrder on the command line or in a module's rules.
