Documentation for build tools

diff --git a/Engine/Documentation/Source/Programming/UnrealBuildSystem/TargetFiles/TargetFilesProperties/TargetFilesProperties.INT.udn b/Engine/Documentation/Source/Programming/UnrealBuildSystem/TargetFiles/TargetFilesProperties/TargetFilesProperties.INT.udn
--- a/Engine/Documentation/Source/Programming/UnrealBuildSystem/TargetFiles/TargetFilesProperties/TargetFilesProperties.INT.udn
+++ b/Engine/Documentation/Source/Programming/UnrealBuildSystem/TargetFiles/TargetFilesProperties/TargetFilesProperties.INT.udn
@@ -356,6 +356,8 @@ $ bDeployAfterCompile (Boolean): Whether to deploy the executable after compilat
 
 $ bAllowRemotelyCompiledPCHs (Boolean): When enabled, allows XGE to compile pre-compiled header files on remote machines.  Otherwise, PCHs are always generated locally.
 
+$ bUseHeaderUnitsForPch (Boolean): Will replace pch with ifc and use header units instead. This is an experimental and msvc-only feature
+
 $ bCheckSystemHeadersForModification (Boolean): Whether headers in system paths should be checked for modification when determining outdated actions.
 
 $ bDisableLinking (Boolean): Whether to disable linking for this target.
