Documentation for build tools

diff --git a/Engine/Documentation/Source/Programming/UnrealBuildSystem/Configuration/BuildConfigProperties/BuildConfigProperties.INT.udn b/Engine/Documentation/Source/Programming/UnrealBuildSystem/Configuration/BuildConfigProperties/BuildConfigProperties.INT.udn
--- a/Engine/Documentation/Source/Programming/UnrealBuildSystem/Configuration/BuildConfigProperties/BuildConfigProperties.INT.udn
+++ b/Engine/Documentation/Source/Programming/UnrealBuildSystem/Configuration/BuildConfigProperties/BuildConfigProperties.INT.udn
@@ -318,6 +318,8 @@ $ bCompilerTrace : Outputs compile timing information so that it can be analyzed
 
 $ bShowIncludes : Print out files that are included by each source file
 
+$ bSetResourceVersions : If enabled will set the ProductVersion embeded in windows executables and dlls to contain BUILT_FROM_CHANGELIST and BuildVersion Enabled by default for all precompiled and Shipping configurations. Regardless of this setting, the versions from Build.version will be available via the BuildSettings module Note: Embedding these versions will cause resource files to be recompiled whenever changelist is updated which will cause binaries to relink
+
 $ InlineFunctionExpansionLevel : Which level to use for Inline Function Expansion when TargetRules.bUseInlining is enabled
 
 $ ToolChain : Version of the toolchain to use on Windows platform when a non-msvc Compiler is in use, to locate include paths etc.
diff --git a/Engine/Documentation/Source/Programming/UnrealBuildSystem/TargetFiles/TargetFilesProperties/TargetFilesProperties.INT.udn b/Engine/Documentation/Source/Programming/UnrealBuildSystem/TargetFiles/TargetFilesProperties/TargetFilesProperties.INT.udn
--- a/Engine/Documentation/Source/Programming/UnrealBuildSystem/TargetFiles/TargetFilesProperties/TargetFilesProperties.INT.udn
+++ b/Engine/Documentation/Source/Programming/UnrealBuildSystem/TargetFiles/TargetFilesProperties/TargetFilesProperties.INT.udn
@@ -364,7 +364,7 @@ $ bDisableLinking (Boolean): Whether to disable linking for this target.
 
 $ bIgnoreBuildOutputs (Boolean): Whether to ignore tracking build outputs for this target.
 
-$ bFormalBuild (Boolean): Indicates that this is a formal build, intended for distribution. This flag is automatically set to true when Build.version has a changelist set. The only behavior currently bound to this flag is to compile the default resource file separately for each binary so that the OriginalFilename field is set correctly. By default, we only compile the resource once to reduce build times.
+$ bFormalBuild (Boolean): Indicates that this is a formal build, intended for distribution. This flag is automatically set to true when Build.version has a changelist set and is a promoted build. The only behavior currently bound to this flag is to compile the default resource file separately for each binary so that the OriginalFilename field is set correctly. By default, we only compile the resource once to reduce build times.
 
 $ bFlushBuildDirOnRemoteMac (Boolean): Whether to clean Builds directory on a remote Mac before building.
 
