Documentation for build tools

diff --git a/Engine/Documentation/Source/Programming/UnrealBuildSystem/Configuration/BuildConfigProperties/BuildConfigProperties.INT.udn b/Engine/Documentation/Source/Programming/UnrealBuildSystem/Configuration/BuildConfigProperties/BuildConfigProperties.INT.udn
--- a/Engine/Documentation/Source/Programming/UnrealBuildSystem/Configuration/BuildConfigProperties/BuildConfigProperties.INT.udn
+++ b/Engine/Documentation/Source/Programming/UnrealBuildSystem/Configuration/BuildConfigProperties/BuildConfigProperties.INT.udn
@@ -10,7 +10,7 @@ $ bIgnoreOutdatedImportLibraries : Whether to ignore import library files that a
 
 $ bPrintDebugInfo : Whether debug info should be written to the console.
 
-$ bAllowHybridExecutor : Whether the hybrid executor will be used (a remote executor and local executor).
+$ bAllowHybridExecutor : Whether the hybrid executor will be used (a remote executor and local executor). No longer supported.
 
 $ bAllowXGE : Whether XGE may be used if available, default is true.
 
