Documentation for build tools

diff --git a/Engine/Documentation/Source/Programming/Development/BuildGraph/ScriptAnatomy/Tasks/PredefinedTasks/BuildGraphScriptPredefinedTasks.INT.udn b/Engine/Documentation/Source/Programming/Development/BuildGraph/ScriptAnatomy/Tasks/PredefinedTasks/BuildGraphScriptPredefinedTasks.INT.udn
--- a/Engine/Documentation/Source/Programming/Development/BuildGraph/ScriptAnatomy/Tasks/PredefinedTasks/BuildGraphScriptPredefinedTasks.INT.udn
+++ b/Engine/Documentation/Source/Programming/Development/BuildGraph/ScriptAnatomy/Tasks/PredefinedTasks/BuildGraphScriptPredefinedTasks.INT.udn
@@ -202,6 +202,8 @@ Spawns Docker and waits for it to complete.
 | UseBuildKit      | Boolean   | Optional | Use BuildKit in Docker                                                                                         |
 | ProgressOutput   | String    | Optional | Type of progress output (--progress)                                                                           |
 | Tag              | String    | Optional | Tag for the image                                                                                              |
+| Target           | String    | Optional | Set the target build stage to build (--target)                                                                 |
+| Output           | String    | Optional | Custom output exporter. Requires BuildKit (--output)                                                           |
 | Arguments        | String    | Optional | Optional arguments                                                                                             |
 | OverlayDirs      | String    | Optional | List of additional directories to overlay into the staged input files. Allows credentials to be staged, etc... |
 | Environment      | String    | Optional | Environment variables to set                                                                                   |
@@ -686,19 +688,20 @@ Writes text to a file.
 
 Exports an snapshot from Zen to a specified destination.
 
-|                            |               |          |                                                                                   |
-| -------------------------- | ------------- | -------- | --------------------------------------------------------------------------------- |
-| Project                    | FileReference | Optional | The project from which to export the snapshot                                     |
-| Platform                   | String        | Optional | The target platform(s) to export the snapshot for                                 |
-| SnapshotDescriptorFile     | FileReference | Optional | A file to create with information about the snapshot that was exported            |
-| DestinationStorageType     | String        | Required | The type of destination to export the snapshot to (cloud, ...)                    |
-| DestinationCloudHost       | String        | Optional | The host name to use when exporting to a cloud destination                        |
-| DestinationCloudNamespace  | String        | Optional | The namespace to use when exporting to a cloud destination                        |
-| DestinationCloudIdentifier | String        | Optional | The identifier to use when exporting to a cloud destination                       |
-| DestinationCloudBucket     | String        | Optional | A custom bucket name to use when exporting to a cloud destination                 |
-| DestinationDir             | String        | Optional | The directory to store the exported oplog                                         |
-| DestinationOplogName       | String        | Optional | The name of the oplog to write on disk                                            |
-| If                         | Condition     | Optional | Whether to execute this task. It is ignored if this condition evaluates to false. |
+|                            |                    |          |                                                                                                                   |
+| -------------------------- | ------------------ | -------- | ----------------------------------------------------------------------------------------------------------------- |
+| Project                    | FileReference      | Optional | The project from which to export the snapshot                                                                     |
+| Platform                   | String             | Optional | The target platform(s) to export the snapshot for                                                                 |
+| SnapshotBaseDescriptorFile | FileReference      | Optional | A file to read with information about the snapshot that should be used as a base when exporting this new snapshot |
+| SnapshotDescriptorFile     | FileReference      | Optional | A file to create with information about the snapshot that was exported                                            |
+| DestinationStorageType     | String             | Required | The type of destination to export the snapshot to (cloud, ...)                                                    |
+| DestinationIdentifier      | String             | Optional | The identifier to use when exporting to a destination                                                             |
+| DestinationCloudHost       | String             | Optional | The host name to use when exporting to a cloud destination                                                        |
+| DestinationCloudNamespace  | String             | Optional | The namespace to use when exporting to a cloud destination                                                        |
+| DestinationCloudBucket     | String             | Optional | A custom bucket name to use when exporting to a cloud destination                                                 |
+| DestinationFileDir         | DirectoryReference | Optional | The directory to use when exporting to a file destination                                                         |
+| DestinationFileName        | String             | Optional | The filename to use when exporting to a file destination                                                          |
+| If                         | Condition          | Optional | Whether to execute this task. It is ignored if this condition evaluates to false.                                 |
 
 ### ZenImportOplog
 
