Documentation for build tools

diff --git a/Engine/Documentation/Source/Programming/UnrealBuildSystem/Configuration/BuildConfigProperties/BuildConfigProperties.INT.udn b/Engine/Documentation/Source/Programming/UnrealBuildSystem/Configuration/BuildConfigProperties/BuildConfigProperties.INT.udn
--- a/Engine/Documentation/Source/Programming/UnrealBuildSystem/Configuration/BuildConfigProperties/BuildConfigProperties.INT.udn
+++ b/Engine/Documentation/Source/Programming/UnrealBuildSystem/Configuration/BuildConfigProperties/BuildConfigProperties.INT.udn
@@ -108,9 +108,7 @@ $ bStressTestUnity : Whether to stress test the C++ unity build robustness by in
 
 $ bDetailedUnityFiles : Whether to add additional information to the unity files, such as '_of_X' in the file name. Not recommended.
 
-$ bDisableDebugInfo : Whether to globally disable debug info generation; see DebugInfoHeuristics.cs for per-config and per-platform options.
-
-$ bDisableDebugInfoForGeneratedCode : Whether to disable debug info generation for generated files. This improves link times for modules that have a lot of generated glue code.
+$ bDisableDebugInfoForGeneratedCode : Whether to disable debug info generation for generated files. This improves link times and reduces pdb size for modules that have a lot of generated glue code.
 
 $ bOmitPCDebugInfoInDevelopment : Whether to disable debug info on PC/Mac in development builds (for faster developer iteration, as link times are extremely fast with debug info disabled).
 
@@ -200,6 +198,12 @@ $ bRequireObjectPtrForAddReferencedObjects : Require TObjectPtr for FReferenceCo
 
 $ bValidateFormatStrings : Emits compilation errors for incorrect UE_LOG format strings.
 
+$ bDisableDebugInfo : Whether to globally disable debug info generation; Obsolete, please use TargetRules.DebugInfoMode instead
+
+$ DebugInfo : How much debug info should be generated. See DebugInfoMode enum for more details
+
+$ DebugInfoLineTablesOnly : True if only debug line number tables should be emitted in debug information for compilers that support doing so. Overrides TargetRules.DebugInfo See https://clang.llvm.org/docs/UsersManual.html#cmdoption-gline-tables-only for more information
+
 $ bDeterministic : Set flags require for deterministic compiling and linking. Enabling deterministic mode for msvc disables codegen multithreading so compiling will be slower
 
 $ CppStandardEngine : Which C++ standard to use for compiling this target (for engine modules)
diff --git a/Engine/Documentation/Source/Programming/UnrealBuildSystem/TargetFiles/TargetFilesProperties/TargetFilesProperties.INT.udn b/Engine/Documentation/Source/Programming/UnrealBuildSystem/TargetFiles/TargetFilesProperties/TargetFilesProperties.INT.udn
--- a/Engine/Documentation/Source/Programming/UnrealBuildSystem/TargetFiles/TargetFilesProperties/TargetFilesProperties.INT.udn
+++ b/Engine/Documentation/Source/Programming/UnrealBuildSystem/TargetFiles/TargetFilesProperties/TargetFilesProperties.INT.udn
@@ -28,6 +28,14 @@ $ Architecture (UnrealArch): Gets the Architecture in the normal case where ther
 
 $ bIWYU (Boolean): Enables "include what you use" by default for modules in this target. Changes the default PCH mode for any module in this project to PCHUsageMode.UseExplicitOrSharedPCHs.
 
+$ DisableDebugInfoModules (HashSet&lt;String&gt;): Modules that should have debug info disabled
+
+$ DisableDebugInfoPlugins (HashSet&lt;String&gt;): Plugins that should have debug info disabled
+
+$ DebugInfoLineTablesOnlyModules (HashSet&lt;String&gt;): Modules that should emit line number tables instead of full debug information for compilers that support doing so. Overrides DisableDebugInfoModules
+
+$ DebugInfoLineTablesOnlyPlugins (HashSet&lt;String&gt;): Plugins that should emit line number tables instead of full debug information for compilers that support doing so. Overrides DisableDebugInfoPlugins
+
 $ HostPlatform (UnrealTargetPlatform): Gets the host platform being built on
 
 $ bGenerateProjectFiles (Boolean): Expose the bGenerateProjectFiles flag to targets, so we can modify behavior as appropriate for better intellisense
@@ -270,9 +278,7 @@ $ bDetailedUnityFiles (Boolean): Whether to add additional information to the un
 
 $ bForceDebugInfo (Boolean): Whether to force debug info to be generated.
 
-$ bDisableDebugInfo (Boolean): Whether to globally disable debug info generation; see DebugInfoHeuristics.cs for per-config and per-platform options.
-
-$ bDisableDebugInfoForGeneratedCode (Boolean): Whether to disable debug info generation for generated files. This improves link times for modules that have a lot of generated glue code.
+$ bDisableDebugInfoForGeneratedCode (Boolean): Whether to disable debug info generation for generated files. This improves link times and reduces pdb size for modules that have a lot of generated glue code.
 
 $ bOmitPCDebugInfoInDevelopment (Boolean): Whether to disable debug info on PC/Mac in development builds (for faster developer iteration, as link times are extremely fast with debug info disabled).
 
@@ -510,6 +516,12 @@ $ DeprecationWarningLevel (WarningLevel): Level to report deprecation warnings a
 
 $ bRetainFramePointers (Boolean): Forces frame pointers to be retained this is usually required when you want reliable callstacks e.g. mallocframeprofiler
 
+$ bDisableDebugInfo (Boolean): Whether to globally disable debug info generation; Obsolete, please use TargetRules.DebugInfoMode instead
+
+$ DebugInfo (DebugInfoMode): How much debug info should be generated. See DebugInfoMode enum for more details
+
+$ DebugInfoLineTablesOnly (DebugInfoMode): True if only debug line number tables should be emitted in debug information for compilers that support doing so. Overrides TargetRules.DebugInfo See https://clang.llvm.org/docs/UsersManual.html#cmdoption-gline-tables-only for more information
+
 $ bDeterministic (Boolean): Set flags require for deterministic compiling and linking. Enabling deterministic mode for msvc disables codegen multithreading so compiling will be slower
 
 $ bLegacyPublicIncludePaths (Boolean): Add all the public folders as include paths for the compile environment.
