Documentation for build tools

diff --git a/Engine/Documentation/Source/Programming/Development/BuildGraph/ScriptAnatomy/Tasks/PredefinedTasks/BuildGraphScriptPredefinedTasks.INT.udn b/Engine/Documentation/Source/Programming/Development/BuildGraph/ScriptAnatomy/Tasks/PredefinedTasks/BuildGraphScriptPredefinedTasks.INT.udn
--- a/Engine/Documentation/Source/Programming/Development/BuildGraph/ScriptAnatomy/Tasks/PredefinedTasks/BuildGraphScriptPredefinedTasks.INT.udn
+++ b/Engine/Documentation/Source/Programming/Development/BuildGraph/ScriptAnatomy/Tasks/PredefinedTasks/BuildGraphScriptPredefinedTasks.INT.udn
@@ -688,20 +688,21 @@ Writes text to a file.
 
 Exports an snapshot from Zen to a specified destination.
 
-|                            |                    |          |                                                                                                                   |
-| -------------------------- | ------------------ | -------- | ----------------------------------------------------------------------------------------------------------------- |
-| Project                    | FileReference      | Optional | The project from which to export the snapshot                                                                     |
-| Platform                   | String             | Optional | The target platform(s) to export the snapshot for                                                                 |
-| SnapshotBaseDescriptorFile | FileReference      | Optional | A file to read with information about the snapshot that should be used as a base when exporting this new snapshot |
-| SnapshotDescriptorFile     | FileReference      | Optional | A file to create with information about the snapshot that was exported                                            |
-| DestinationStorageType     | String             | Required | The type of destination to export the snapshot to (cloud, ...)                                                    |
-| DestinationIdentifier      | String             | Optional | The identifier to use when exporting to a destination                                                             |
-| DestinationCloudHost       | String             | Optional | The host name to use when exporting to a cloud destination                                                        |
-| DestinationCloudNamespace  | String             | Optional | The namespace to use when exporting to a cloud destination                                                        |
-| DestinationCloudBucket     | String             | Optional | A custom bucket name to use when exporting to a cloud destination                                                 |
-| DestinationFileDir         | DirectoryReference | Optional | The directory to use when exporting to a file destination                                                         |
-| DestinationFileName        | String             | Optional | The filename to use when exporting to a file destination                                                          |
-| If                         | Condition          | Optional | Whether to execute this task. It is ignored if this condition evaluates to false.                                 |
+|                            |                    |          |                                                                                                                            |
+| -------------------------- | ------------------ | -------- | -------------------------------------------------------------------------------------------------------------------------- |
+| Project                    | FileReference      | Optional | The project from which to export the snapshot                                                                              |
+| Platform                   | String             | Optional | The target platform(s) to export the snapshot for                                                                          |
+| SnapshotBaseDescriptorFile | FileReference      | Optional | A file to read with information about the snapshot that should be used as a base when exporting this new snapshot          |
+| SnapshotDescriptorFile     | FileReference      | Optional | A file to create with information about the snapshot that was exported                                                     |
+| DestinationStorageType     | String             | Required | The type of destination to export the snapshot to (cloud, ...)                                                             |
+| DestinationIdentifier      | String             | Optional | The identifier to use when exporting to a destination                                                                      |
+| DestinationCloudHost       | String             | Optional | The host name to use when exporting to a cloud destination                                                                 |
+| DestinationCloudNamespace  | String             | Optional | The namespace to use when exporting to a cloud destination                                                                 |
+| DestinationCloudBucket     | String             | Optional | A custom bucket name to use when exporting to a cloud destination                                                          |
+| DestinationFileDir         | DirectoryReference | Optional | The directory to use when exporting to a file destination                                                                  |
+| DestinationFileName        | String             | Optional | The filename to use when exporting to a file destination                                                                   |
+| OverridePlatformCookedDir  | String             | Optional | Optional. Where to look for the .projectstore The pattern {Platform} can be used for exporting multiple platforms at once. |
+| If                         | Condition          | Optional | Whether to execute this task. It is ignored if this condition evaluates to false.                                          |
 
 ### ZenImportOplog
 
