Documentation for build tools

diff --git a/Engine/Documentation/Source/Programming/UnrealBuildSystem/Configuration/BuildConfigProperties/BuildConfigProperties.INT.udn b/Engine/Documentation/Source/Programming/UnrealBuildSystem/Configuration/BuildConfigProperties/BuildConfigProperties.INT.udn
--- a/Engine/Documentation/Source/Programming/UnrealBuildSystem/Configuration/BuildConfigProperties/BuildConfigProperties.INT.udn
+++ b/Engine/Documentation/Source/Programming/UnrealBuildSystem/Configuration/BuildConfigProperties/BuildConfigProperties.INT.udn
@@ -12,6 +12,8 @@ $ bPrintDebugInfo : Whether debug info should be written to the console.
 
 $ bAllowHybridExecutor : Whether the hybrid executor will be used (a remote executor and local executor). No longer supported.
 
+$ bAllowBoxExecutor : Whether the experimental box executor will be used.
+
 $ bAllowXGE : Whether XGE may be used if available, default is true.
 
 $ bAllowFASTBuild : Whether FASTBuild may be used if availabe, default is true.
@@ -46,6 +48,8 @@ $ bUseUnityBuild : Whether to unify C++ code into larger files for faster compil
 
 $ bForceUnityBuild : Whether to force C++ source files to be combined into larger files for faster compilation.
 
+$ bAllowBoxLocalExecutor : Whether the experimental box (local only) executor will be used.
+
 $ bDebugBuildsActuallyUseDebugCRT : Enables the debug C++ runtime (CRT) for debug builds. By default we always use the release runtime, since the debug version isn't particularly useful when debugging Unreal Engine projects, and linking against the debug CRT libraries forces our third party library dependencies to also be compiled using the debug CRT (and often perform more slowly). Often it can be inconvenient to require a separate copy of the debug versions of third party static libraries simply so that you can debug your program's code.
 
 $ bLegalToDistributeBinary : Whether the output from this target can be publicly distributed, even if it has dependencies on modules that are in folders with special restrictions (eg. CarefullyRedist, NotForLicensees, NoRedist).
