Documentation for build tools

diff --git a/Engine/Documentation/Source/Programming/UnrealBuildSystem/TargetFiles/TargetFilesProperties/TargetFilesProperties.INT.udn b/Engine/Documentation/Source/Programming/UnrealBuildSystem/TargetFiles/TargetFilesProperties/TargetFilesProperties.INT.udn
--- a/Engine/Documentation/Source/Programming/UnrealBuildSystem/TargetFiles/TargetFilesProperties/TargetFilesProperties.INT.udn
+++ b/Engine/Documentation/Source/Programming/UnrealBuildSystem/TargetFiles/TargetFilesProperties/TargetFilesProperties.INT.udn
@@ -178,6 +178,8 @@ $ bWithFixedTimeStepSupport (Boolean): Whether to enable the support for FixedTi
 
 $ bUseLoggingInShipping (Boolean): Whether to turn on logging for test/shipping builds.
 
+$ bUseConsoleInShipping (Boolean): Whether to turn on console for shipping builds.
+
 $ bLoggingToMemoryEnabled (Boolean): Whether to turn on logging to memory for test/shipping builds.
 
 $ bUseLauncherChecks (Boolean): Whether to check that the process was launched through an external launcher.
