Documentation for build tools

diff --git a/Engine/Documentation/Source/Programming/UnrealBuildSystem/Configuration/BuildConfigProperties/BuildConfigProperties.INT.udn b/Engine/Documentation/Source/Programming/UnrealBuildSystem/Configuration/BuildConfigProperties/BuildConfigProperties.INT.udn
--- a/Engine/Documentation/Source/Programming/UnrealBuildSystem/Configuration/BuildConfigProperties/BuildConfigProperties.INT.udn
+++ b/Engine/Documentation/Source/Programming/UnrealBuildSystem/Configuration/BuildConfigProperties/BuildConfigProperties.INT.udn
@@ -484,6 +484,8 @@ $ bIncludeDebugConfigs : Whether we should include configurations for "Debug" an
 
 $ bIncludeDevelopmentConfigs : Whether we should include configurations for "Development" in generated projects. Pass "-NoDevelopmentConfigs" to disable this.
 
+$ bVisualStudioLinux : True if visual studio project should be generated in linux mode.
+
 ### IOSToolChain
 
 $ bUseDangerouslyFastMode : If this is set, then we do not do any post-compile steps -- except moving the executable into the proper spot on Mac.
