Documentation for build tools

diff --git a/Engine/Documentation/Source/Programming/UnrealBuildSystem/Configuration/BuildConfigProperties/BuildConfigProperties.INT.udn b/Engine/Documentation/Source/Programming/UnrealBuildSystem/Configuration/BuildConfigProperties/BuildConfigProperties.INT.udn
--- a/Engine/Documentation/Source/Programming/UnrealBuildSystem/Configuration/BuildConfigProperties/BuildConfigProperties.INT.udn
+++ b/Engine/Documentation/Source/Programming/UnrealBuildSystem/Configuration/BuildConfigProperties/BuildConfigProperties.INT.udn
@@ -310,7 +310,7 @@ $ bUpdatedCPPMacro : Enables updated __cplusplus macro (/Zc:__cplusplus).
 
 $ bStrictInlineConformance : Enables inline conformance (Remove unreferenced COMDAT) (/Zc:inline).
 
-$ bStrictPreprocessorConformance : Enables new preprocessor conformance (/Zc:preprocessor).
+$ bStrictPreprocessorConformance : Enables new preprocessor conformance (/Zc:preprocessor). This is always enabled for C++20 modules.
 
 $ bStrictEnumTypesConformance : Enables enum types conformance (/Zc:enumTypes) in VS2022 17.4 Preview 4.0+.
 
