Documentation for build tools

diff --git a/Engine/Documentation/Source/Programming/UnrealBuildSystem/TargetFiles/TargetFilesProperties/TargetFilesProperties.INT.udn b/Engine/Documentation/Source/Programming/UnrealBuildSystem/TargetFiles/TargetFilesProperties/TargetFilesProperties.INT.udn
--- a/Engine/Documentation/Source/Programming/UnrealBuildSystem/TargetFiles/TargetFilesProperties/TargetFilesProperties.INT.udn
+++ b/Engine/Documentation/Source/Programming/UnrealBuildSystem/TargetFiles/TargetFilesProperties/TargetFilesProperties.INT.udn
@@ -82,6 +82,8 @@ $ EnablePlugins (List&lt;String&gt;): Additional plugins that should be included
 
 $ DisablePlugins (List&lt;String&gt;): List of plugins to be disabled for this target. Note that the project file may still reference them, so they should be marked as optional to avoid failing to find them at runtime.
 
+$ OptionalPlugins (List&lt;String&gt;): Additional plugins that should be included for this target if they are found.
+
 $ DisablePluginsConflictWarningLevel (WarningLevel): How to treat conflicts when a disabled plugin is being enabled by another plugin referencing it
 
 $ InternalPluginDependencies (List&lt;String&gt;): A list of Plugin names that are allowed to exist as dependencies without being defined in the uplugin descriptor
