Documentation for build tools

diff --git a/Engine/Documentation/Source/Programming/UnrealBuildSystem/Configuration/BuildConfigProperties/BuildConfigProperties.INT.udn b/Engine/Documentation/Source/Programming/UnrealBuildSystem/Configuration/BuildConfigProperties/BuildConfigProperties.INT.udn
--- a/Engine/Documentation/Source/Programming/UnrealBuildSystem/Configuration/BuildConfigProperties/BuildConfigProperties.INT.udn
+++ b/Engine/Documentation/Source/Programming/UnrealBuildSystem/Configuration/BuildConfigProperties/BuildConfigProperties.INT.udn
@@ -306,6 +306,8 @@ $ ToolchainVersion : The specific msvc toolchain version to use if the compiler
 
 $ bVCFastFail : True if /fastfail should be passed to the msvc compiler and linker
 
+$ bVCExtendedWarningInfo : True if /d2ExtendedWarningInfo should be passed to the compiler and /d2:-ExtendedWarningInfo to the linker
+
 $ bClangStandaloneDebug : True if optimizations to reduce the size of debug information should be disabled See https://clang.llvm.org/docs/UsersManual.html#cmdoption-fstandalone-debug for more information
 
 $ bAllowClangLinker : True if we should use the Clang linker (LLD) when we are compiling with Clang, or Intel linker (xilink\xilib) when we are compiling with Intel oneAPI, otherwise we use the MSVC linker.
