Documentation for build tools

diff --git a/Engine/Documentation/Source/Programming/UnrealBuildSystem/Configuration/BuildConfigProperties/BuildConfigProperties.INT.udn b/Engine/Documentation/Source/Programming/UnrealBuildSystem/Configuration/BuildConfigProperties/BuildConfigProperties.INT.udn
--- a/Engine/Documentation/Source/Programming/UnrealBuildSystem/Configuration/BuildConfigProperties/BuildConfigProperties.INT.udn
+++ b/Engine/Documentation/Source/Programming/UnrealBuildSystem/Configuration/BuildConfigProperties/BuildConfigProperties.INT.udn
@@ -392,6 +392,86 @@ $ MsvcCRTRedistVersion : Which MSVC CRT Redist version to use
 
 $ CompilerVersion : Which MSVC Compiler version to use
 
+### UnrealBuildAccelerator
+
+$ bDisableRemote : When set to true, UBA will not use any remote help
+
+$ bForceBuildAllRemote : When set to true, UBA will force all actions that can be built remotely to be built remotely. This will hang if there are no remote agents available
+
+$ bForcedRetry : When set to true, actions that fail locally with UBA will be retried without UBA.
+
+$ bStrict : When set to true, all errors and warnings from UBA will be output at the appropriate severity level to the log (rather than being output as 'information' and attempting to continue regardless).
+
+$ bStoreRaw : If UBA should store cas compressed or raw
+
+$ bLinkRemote : If UBA should distribute linking to remote workers. This needs bandwidth but can be an optimization
+
+$ StoreCapacityGb : The amount of gigabytes UBA is allowed to use to store workset and cached data. It is a good idea to have this >10gb
+
+$ MaxWorkers : Max number of worker threads that can handle messages from remotes.
+
+$ SendSize : Max size of each message sent from server to client
+
+$ Host : Which ip UBA server should listen to for connections
+
+$ Port : Which port UBA server should listen to for connections.
+
+$ RootDir : Which directory to store files for UBA.
+
+$ bUseQuic : Use Quic protocol instead of Tcp (experimental)
+
+$ bLogEnabled : Enable logging of UBA processes
+
+$ bPrintSummary : Prints summary of UBA stats at end of build
+
+$ bLaunchVisualizer : Launch visualizer application which shows build progress
+
+$ bResetCas : Resets the cas cache
+
+$ TraceFile : Provide custom path for trace output file
+
+$ bDetailedTrace : Add verbose details to the UBA trace
+
+$ bDisableWaitOnMem : Disable UBA waiting on available memory before spawning new processes
+
+$ bAllowKillOnMem : Let UBA kill running processes when close to out of memory
+
+$ OutputStatsThresholdMs : Threshold for when executor should output logging for the process. Defaults to never
+
+$ bWriteToDisk : Skip writing intermediate and output files to disk. Useful for validation builds where we don't need the output
+
+$ bDisableCustomAlloc : Set to true to disable mimalloc and detouring of memory allocations.
+
+$ Zone : The zone to use for UBA.
+
+$ bUseCrypto : Set to true to enable encryption when transfering files over the network.
+
+$ HordeMaxCores : Max cores allowed to be used by build session
+
+$ HordeDelay : How long UBT should wait to ask for help. Useful in build configs where machine can delay remote work and still get same wall time results (pch dependencies etc)
+
+$ bHordeAllowWine : Allow use of Wine. Only applicable to Horde agents running Linux. Can still be ignored if Wine executable is not set on agent.
+
+$ bDisableHorde : Disable horde all together
+
+### Horde
+
+$ Server : Uri of the Horde server
+
+$ Token : Uri of the Horde server
+
+$ OidcProvider : OIDC id for the login to use
+
+$ Pool : Pool for the Horde agent to assign
+
+$ Requirements : Requirements for the Horde agent to assign
+
+$ LocalHost : Which ip UBA server should give to agents. This will invert so host listens and agents connect
+
+$ ConnectionMode : Connection mode for agent/compute communication for valid modes.
+
+$ UBASentryUrl : Sentry URL to send box data to. Optional.
+
 ### ParallelExecutor
 
 $ MaxProcessorCount : Maximum processor count for local execution.
