Documentation for build tools

diff --git a/Engine/Documentation/Source/Programming/UnrealBuildSystem/Configuration/BuildConfigProperties/BuildConfigProperties.INT.udn b/Engine/Documentation/Source/Programming/UnrealBuildSystem/Configuration/BuildConfigProperties/BuildConfigProperties.INT.udn
--- a/Engine/Documentation/Source/Programming/UnrealBuildSystem/Configuration/BuildConfigProperties/BuildConfigProperties.INT.udn
+++ b/Engine/Documentation/Source/Programming/UnrealBuildSystem/Configuration/BuildConfigProperties/BuildConfigProperties.INT.udn
@@ -470,6 +470,8 @@ $ LocalHost : Which ip UBA server should give to agents. This will invert so hos
 
 $ ConnectionMode : Connection mode for agent/compute communication for valid modes.
 
+$ Encryption : Encryption to use for agent/compute communication. Note that UBA agent uses its own encryption. for valid modes.
+
 $ UBASentryUrl : Sentry URL to send box data to. Optional.
 
 ### ParallelExecutor
