Documentation for build tools

diff --git a/Engine/Documentation/Source/Programming/UnrealBuildSystem/ModuleFiles/ModuleFilesProperties/ModuleFilesProperties.INT.udn b/Engine/Documentation/Source/Programming/UnrealBuildSystem/ModuleFiles/ModuleFilesProperties/ModuleFilesProperties.INT.udn
--- a/Engine/Documentation/Source/Programming/UnrealBuildSystem/ModuleFiles/ModuleFilesProperties/ModuleFilesProperties.INT.udn
+++ b/Engine/Documentation/Source/Programming/UnrealBuildSystem/ModuleFiles/ModuleFilesProperties/ModuleFilesProperties.INT.udn
@@ -192,6 +192,8 @@ $ PCHUsage (PCHUsageMode): Precompiled header usage for this module
 
 $ bValidateFormatStrings (Boolean): Emits compilation errors for incorrect UE_LOG format strings.
 
+$ bValidateInternalApi (Boolean): Emits deprecated warnings\errors for internal API usage for non-engine modules
+
 $ DefaultBuildSettings (BuildSettingsVersion): Which engine version's build settings to use by default.
 
 $ IncludeOrderVersion (EngineIncludeOrderVersion): What version of include order to use when compiling this module. Can be overridden via -ForceIncludeOrder on the command line or in a module's rules.
