Documentation for build tools

diff --git a/Engine/Documentation/Source/Programming/Development/BuildGraph/ScriptAnatomy/Tasks/PredefinedTasks/BuildGraphScriptPredefinedTasks.INT.udn b/Engine/Documentation/Source/Programming/Development/BuildGraph/ScriptAnatomy/Tasks/PredefinedTasks/BuildGraphScriptPredefinedTasks.INT.udn
--- a/Engine/Documentation/Source/Programming/Development/BuildGraph/ScriptAnatomy/Tasks/PredefinedTasks/BuildGraphScriptPredefinedTasks.INT.udn
+++ b/Engine/Documentation/Source/Programming/Development/BuildGraph/ScriptAnatomy/Tasks/PredefinedTasks/BuildGraphScriptPredefinedTasks.INT.udn
@@ -132,6 +132,18 @@ Copies files from one directory to another.
 | ErrorIfNotFound | Boolean   | Optional | Whether or not to throw an error if no files were found to copy                   |
 | If              | Condition | Optional | Whether to execute this task. It is ignored if this condition evaluates to false. |
 
+### CreateArtifact
+
+Deploys a tool update through Horde
+
+|         |           |          |                                                                                   |
+| ------- | --------- | -------- | --------------------------------------------------------------------------------- |
+| Name    | String    | Required | Name for the artifact                                                             |
+| Type    | String    | Optional | Type of the artifact                                                              |
+| BaseDir | String    | Required | Base directory to resolve relative paths for input files.                         |
+| Files   | FileSpec  | Optional | Files to be uploaded.                                                             |
+| If      | Condition | Optional | Whether to execute this task. It is ignored if this condition evaluates to false. |
+
 ### CsCompile
 
 Compiles C# project files, and their dependencies.
