Documentation for build tools

diff --git a/Engine/Documentation/Source/Programming/UnrealBuildSystem/ModuleFiles/ModuleFilesProperties/ModuleFilesProperties.INT.udn b/Engine/Documentation/Source/Programming/UnrealBuildSystem/ModuleFiles/ModuleFilesProperties/ModuleFilesProperties.INT.udn
--- a/Engine/Documentation/Source/Programming/UnrealBuildSystem/ModuleFiles/ModuleFilesProperties/ModuleFilesProperties.INT.udn
+++ b/Engine/Documentation/Source/Programming/UnrealBuildSystem/ModuleFiles/ModuleFilesProperties/ModuleFilesProperties.INT.udn
@@ -40,6 +40,8 @@ $ bTreatAsEngineModule (Boolean): Whether this module should be treated as an en
 
 $ bUseRTTI (Boolean): Use run time type information
 
+$ bVcRemoveUnreferencedComdat (Boolean): Whether to direct MSVC to remove unreferenced COMDAT functions and data.
+
 $ bCodeCoverage (Boolean): Enable code coverage compilation/linking support.
 
 $ MinCpuArchX64 (Nullable&lt;MinimumCpuArchitectureX64&gt;): Direct the compiler to generate AVX instructions wherever SSE or AVX intrinsics are used, on the x64 platforms that support it. Note that by enabling this you are changing the minspec for the PC platform, and the resultant executable will crash on machines without AVX support.
