Documentation for build tools

diff --git a/Engine/Documentation/Source/Programming/UnrealBuildSystem/Configuration/BuildConfigProperties/BuildConfigProperties.INT.udn b/Engine/Documentation/Source/Programming/UnrealBuildSystem/Configuration/BuildConfigProperties/BuildConfigProperties.INT.udn
--- a/Engine/Documentation/Source/Programming/UnrealBuildSystem/Configuration/BuildConfigProperties/BuildConfigProperties.INT.udn
+++ b/Engine/Documentation/Source/Programming/UnrealBuildSystem/Configuration/BuildConfigProperties/BuildConfigProperties.INT.udn
@@ -298,6 +298,8 @@ $ bUseFastGenProfile : If specified along with -PGOProfile, then /FASTGENPROFILE
 
 $ bPGONoExtraCounters : If specified along with -PGOProfile, prevent the usage of extra counters. Please note that by default /FASTGENPROFILE doesnt use extra counters
 
+$ bSampleBasedPGO : If specified along with -PGOProfile, use sample-based PGO instead of instrumented. Currently Intel oneAPI 2024.0+ only.
+
 $ Compiler : Version of the compiler toolchain to use on Windows platform. A value of "default" will be changed to a specific version at UBT start up.
 
 $ CompilerVersion : The specific compiler version to use. This may be a specific version number (for example, "14.13.26128"), the string "Latest" to select the newest available version, or the string "Preview" to select the newest available preview version. By default, and if it is available, we use the toolchain version indicated by WindowsPlatform.DefaultToolChainVersion (otherwise, we use the latest version).
@@ -310,7 +312,7 @@ $ bVCExtendedWarningInfo : True if /d2ExtendedWarningInfo should be passed to th
 
 $ bClangStandaloneDebug : True if optimizations to reduce the size of debug information should be disabled See https://clang.llvm.org/docs/UsersManual.html#cmdoption-fstandalone-debug for more information
 
-$ bAllowClangLinker : True if we should use the Clang linker (LLD) when we are compiling with Clang, or Intel linker (xilink\xilib) when we are compiling with Intel oneAPI, otherwise we use the MSVC linker.
+$ bAllowClangLinker : True if we should use the Clang linker (LLD) when we are compiling with Clang or Intel oneAPI, otherwise we use the MSVC linker.
 
 $ WindowsSdkVersion : The specific Windows SDK version to use. This may be a specific version number (for example, "8.1", "10.0" or "10.0.10150.0"), or the string "Latest", to select the newest available version. By default, and if it is available, we use the Windows SDK version indicated by WindowsPlatform.DefaultWindowsSdkVersion (otherwise, we use the latest version).
 
@@ -446,6 +448,8 @@ $ Zone : The zone to use for UBA.
 
 $ bUseCrypto : Set to true to enable encryption when transfering files over the network.
 
+$ bUseKnownInputs : Set to true to provide known inputs to processes that are run remote. This is an experimental feature to speed up build times when ping is higher
+
 $ HordeMaxCores : Max cores allowed to be used by build session
 
 $ HordeDelay : How long UBT should wait to ask for help. Useful in build configs where machine can delay remote work and still get same wall time results (pch dependencies etc)
