Documentation for build tools

diff --git a/Engine/Documentation/Source/Programming/UnrealBuildSystem/Configuration/BuildConfigProperties/BuildConfigProperties.INT.udn b/Engine/Documentation/Source/Programming/UnrealBuildSystem/Configuration/BuildConfigProperties/BuildConfigProperties.INT.udn
--- a/Engine/Documentation/Source/Programming/UnrealBuildSystem/Configuration/BuildConfigProperties/BuildConfigProperties.INT.udn
+++ b/Engine/Documentation/Source/Programming/UnrealBuildSystem/Configuration/BuildConfigProperties/BuildConfigProperties.INT.udn
@@ -56,6 +56,10 @@ $ bAllowUBAExecutor : Whether the experimental UnrealBuildAccelerator executor w
 
 $ bAllowUBALocalExecutor : Whether the experimental UnrealBuildAccelerator (local only) executor will be used.
 
+$ DefaultWarningLevel : Default treatment of uncategorized warnings
+
+$ DeprecationWarningLevel : Level to report deprecation warnings as errors
+
 $ bDebugBuildsActuallyUseDebugCRT : Enables the debug C++ runtime (CRT) for debug builds. By default we always use the release runtime, since the debug version isn't particularly useful when debugging Unreal Engine projects, and linking against the debug CRT libraries forces our third party library dependencies to also be compiled using the debug CRT (and often perform more slowly). Often it can be inconvenient to require a separate copy of the debug versions of third party static libraries simply so that you can debug your program's code.
 
 $ bLegalToDistributeBinary : Whether the output from this target can be publicly distributed, even if it has dependencies on modules that are in folders with special restrictions (eg. CarefullyRedist, NotForLicensees, NoRedist).
@@ -70,8 +74,6 @@ $ bUseDebugLiveCodingConsole : Whether to enable support for live coding
 
 $ bUseXGEController : Whether the XGE controller worker and modules should be included in the engine build. These are required for distributed shader compilation using the XGE interception interface.
 
-$ bMergeModuleAndGeneratedUnityFiles : Whether to merge module and generated unity files for faster compilation.
-
 $ bUseAdaptiveUnityBuild : Use a heuristic to determine which files are currently being iterated on and exclude them from unity blobs, result in faster incremental compile times. The current implementation uses the read-only flag to distinguish the working set, assuming that files will be made writable by the source control system if they are being modified. This is true for Perforce, but not for Git.
 
 $ bAdaptiveUnityDisablesOptimizations : Disable optimization for files that are in the adaptive non-unity working set.
@@ -88,9 +90,9 @@ $ bAdaptiveUnityCompilesHeaderFiles : Creates a dedicated source file for each h
 
 $ MinGameModuleSourceFilesForUnityBuild : The number of source files in a game module before unity build will be activated for that module.  This allows small game modules to have faster iterative compile times for single files, at the expense of slower full rebuild times.  This setting can be overridden by the bFasterWithoutUnity option in a module's Build.cs file.
 
-$ DefaultWarningLevel : Default treatment of uncategorized warnings
+$ bRequireObjectPtrForAddReferencedObjects : Require TObjectPtr for FReferenceCollector API's. (Needed for compatibility with incremental GC.)
 
-$ DeprecationWarningLevel : Level to report deprecation warnings as errors
+$ bValidateFormatStrings : Emits compilation errors for incorrect UE_LOG format strings.
 
 $ bWarningsAsErrors : Whether to enable all warnings as errors. UE enables most warnings as errors already, but disables a few (such as deprecation warnings).
 
@@ -118,6 +120,12 @@ $ bStressTestUnity : Whether to stress test the C++ unity build robustness by in
 
 $ bDetailedUnityFiles : Whether to add additional information to the unity files, such as '_of_X' in the file name. Not recommended.
 
+$ bDisableDebugInfo : Whether to globally disable debug info generation; Obsolete, please use TargetRules.DebugInfoMode instead
+
+$ DebugInfo : How much debug info should be generated. See DebugInfoMode enum for more details
+
+$ DebugInfoLineTablesOnly : True if only debug line number tables should be emitted in debug information for compilers that support doing so. Overrides TargetRules.DebugInfo See https://clang.llvm.org/docs/UsersManual.html#cmdoption-gline-tables-only for more information
+
 $ bDisableDebugInfoForGeneratedCode : Whether to disable debug info generation for generated files. This improves link times and reduces pdb size for modules that have a lot of generated glue code.
 
 $ bOmitPCDebugInfoInDevelopment : Whether to disable debug info on PC/Mac in development builds (for faster developer iteration, as link times are extremely fast with debug info disabled).
@@ -126,6 +134,8 @@ $ bUsePDBFiles : Whether PDB files should be used for Visual C++ builds.
 
 $ bUsePCHFiles : Whether PCH files should be used.
 
+$ bDeterministic : Set flags require for deterministic compiling and linking. Enabling deterministic mode for msvc disables codegen multithreading so compiling will be slower
+
 $ bChainPCHs : Whether PCHs should be chained when compiling with clang.
 
 $ bForceIncludePCHHeadersForGenCppFilesWhenPCHIsDisabled : Whether PCH headers should be force included for gen.cpp files when PCH is disabled.
@@ -162,8 +172,6 @@ $ bSupportEditAndContinue : Whether to support edit and continue.
 
 $ bOmitFramePointers : Whether to omit frame pointers or not. Disabling is useful for e.g. memory profiling on the PC.
 
-$ bUseMallocProfiler : If true, then enable memory profiling in the build (defines USE_MALLOC_PROFILER=1 and forces bOmitFramePointers=false).
-
 $ bShaderCompilerWorkerTrace : If true, then enable Unreal Insights (utrace) profiling in the build for the Shader Compiler Worker (defines USE_SHADER_COMPILER_WORKER_TRACE=1).
 
 $ bUseSharedPCHs : Enables "Shared PCHs", a feature which significantly speeds up compile times by attempting to share certain PCH files between modules that UBT detects is including those PCH's header files.
@@ -200,26 +208,14 @@ $ MSVCCompileActionWeight : The weight(cpu/memory utilization) of a MSVC compile
 
 $ ClangCompileActionWeight : The weight(cpu/memory utilization) of a clang compile action
 
-$ CStandard : Which C standard to use for compiling this target
-
-$ MinCpuArchX64 : Direct the compiler to generate AVX instructions wherever SSE or AVX intrinsics are used, on the x64 platforms that support it. Note that by enabling this you are changing the minspec for the PC platform, and the resultant executable will crash on machines without AVX support.
-
-$ bRequireObjectPtrForAddReferencedObjects : Require TObjectPtr for FReferenceCollector API's. (Needed for compatibility with incremental GC.)
-
-$ bValidateFormatStrings : Emits compilation errors for incorrect UE_LOG format strings.
-
-$ bDisableDebugInfo : Whether to globally disable debug info generation; Obsolete, please use TargetRules.DebugInfoMode instead
-
-$ DebugInfo : How much debug info should be generated. See DebugInfoMode enum for more details
-
-$ DebugInfoLineTablesOnly : True if only debug line number tables should be emitted in debug information for compilers that support doing so. Overrides TargetRules.DebugInfo See https://clang.llvm.org/docs/UsersManual.html#cmdoption-gline-tables-only for more information
-
-$ bDeterministic : Set flags require for deterministic compiling and linking. Enabling deterministic mode for msvc disables codegen multithreading so compiling will be slower
-
 $ CppStandardEngine : Which C++ standard to use for compiling this target (for engine modules)
 
 $ CppStandard : Which C++ standard to use for compiling this target (for non-engine modules)
 
+$ CStandard : Which C standard to use for compiling this target
+
+$ MinCpuArchX64 : Direct the compiler to generate AVX instructions wherever SSE or AVX intrinsics are used, on the x64 platforms that support it. Note that by enabling this you are changing the minspec for the PC platform, and the resultant executable will crash on machines without AVX support.
+
 $ ActionStallReportTime : Number of second of no completed actions to trigger an action stall report. If zero, stall reports will not be enabled.
 
 $ bStopSNDBSCompilationAfterErrors : When enabled, SN-DBS will stop compiling targets after a compile error occurs.  Recommended, as it saves computing resources for others.
@@ -344,12 +340,6 @@ $ ToolChain : Version of the toolchain to use on Windows platform when a non-msv
 
 $ bStrictConformanceMode : Enables strict standard conformance mode (/permissive-).
 
-### TargetRules
-
-$ bCompileChaosVisualDebuggerSupport : Whether to compile in Chaos Visual Debugger (CVD) support features to record the state of the physics simulation
-
-$ bDisableInliningGenCpps : Disable supports for inlining gen.cpps
-
 ### ModuleConfiguration
 
 $ DisableMergingModuleAndGeneratedFilesInUnityFiles : List of modules that disables merging module and generated cpp files in the same unity files.
@@ -364,11 +354,13 @@ $ OptimizeForSize : List of modules to optimize for size. This allows the target
 
 $ OptimizeForSizeAndSpeed : List of modules to optimize for size and speed. This allows the target to override module optimization level Note that this may disable PCH usage if a private PCH is not provided
 
-### TestTargetRules
+### TargetRules
 
 $ bCompileChaosVisualDebuggerSupport : Whether to compile in Chaos Visual Debugger (CVD) support features to record the state of the physics simulation
 
-$ bDisableInliningGenCpps : Disable supports for inlining gen.cpps
+### TestTargetRules
+
+$ bCompileChaosVisualDebuggerSupport : Whether to compile in Chaos Visual Debugger (CVD) support features to record the state of the physics simulation
 
 ### FASTBuild
 
diff --git a/Engine/Documentation/Source/Programming/UnrealBuildSystem/TargetFiles/TargetFilesProperties/TargetFilesProperties.INT.udn b/Engine/Documentation/Source/Programming/UnrealBuildSystem/TargetFiles/TargetFilesProperties/TargetFilesProperties.INT.udn
--- a/Engine/Documentation/Source/Programming/UnrealBuildSystem/TargetFiles/TargetFilesProperties/TargetFilesProperties.INT.udn
+++ b/Engine/Documentation/Source/Programming/UnrealBuildSystem/TargetFiles/TargetFilesProperties/TargetFilesProperties.INT.udn
@@ -6,16 +6,6 @@
 
 ### Read-Only Properties
 
-$ Platform (UnrealTargetPlatform): Platform that this target is being built for.
-
-$ Configuration (UnrealTargetConfiguration): The configuration being built.
-
-$ Architectures (UnrealArchitectures): Architecture that the target is being built for (or an empty string for the default).
-
-$ ProjectFile (FileReference): Path to the project file for the project containing this target.
-
-$ Version (ReadOnlyBuildVersion): The current build version
-
 $ OriginalName (String): Return the original target name without overrides or adornments
 
 $ IsTestTarget (Boolean): Whether this is a low level tests target.
@@ -36,6 +26,24 @@
 
 $ DebugInfoLineTablesOnlyPlugins (HashSet&lt;String&gt;): Plugins that should emit line number tables instead of full debug information for compilers that support doing so. Overrides DisableDebugInfoPlugins
 
+$ GlobalDefinitions (List&lt;String&gt;): Macros to define globally across the whole target.
+
+$ ProjectDefinitions (List&lt;String&gt;): Macros to define across all macros in the project.
+
+$ ExtraModuleNames (List&lt;String&gt;): List of additional modules to be compiled into the target.
+
+$ ManifestFileNames (List&lt;FileReference&gt;): Path to a manifest to output for this target
+
+$ DependencyListFileNames (List&lt;FileReference&gt;): Path to a list of dependencies for this target, when precompiling
+
+$ PreBuildTargets (List&lt;TargetInfo&gt;): Specifies a list of targets which should be built before this target is built.
+
+$ PreBuildSteps (List&lt;String&gt;): Specifies a list of steps which should be executed before this target is built, in the context of the host platform's shell. The following variables will be expanded before execution: $(EngineDir), $(ProjectDir), $(TargetName), $(TargetPlatform), $(TargetConfiguration), $(TargetType), $(ProjectFile).
+
+$ PostBuildSteps (List&lt;String&gt;): Specifies a list of steps which should be executed after this target is built, in the context of the host platform's shell. The following variables will be expanded before execution: $(EngineDir), $(ProjectDir), $(TargetName), $(TargetPlatform), $(TargetConfiguration), $(TargetType), $(ProjectFile).
+
+$ AdditionalBuildProducts (List&lt;String&gt;): Specifies additional build products produced as part of this target.
+
 $ HostPlatform (UnrealTargetPlatform): Gets the host platform being built on
 
 $ bGenerateProjectFiles (Boolean): Expose the bGenerateProjectFiles flag to targets, so we can modify behavior as appropriate for better intellisense
@@ -46,12 +54,60 @@
 
 
 ### Read/Write Properties
+$ BuildPlugins (List&lt;String&gt;): Additional plugins that are built for this target type but not enabled.
+
+$ AdditionalPlugins (List&lt;String&gt;): A list of additional plugins which need to be included in this target. This allows referencing non-optional plugin modules which cannot be disabled, and allows building against specific modules in program targets which do not fit the categories in ModuleHostType.
+
+$ EnablePlugins (List&lt;String&gt;): Additional plugins that should be included for this target.
+
+$ DisablePlugins (List&lt;String&gt;): List of plugins to be disabled for this target. Note that the project file may still reference them, so they should be marked as optional to avoid failing to find them at runtime.
+
+$ OptionalPlugins (List&lt;String&gt;): Additional plugins that should be included for this target if they are found.
+
+$ InternalPluginDependencies (List&lt;String&gt;): A list of Plugin names that are allowed to exist as dependencies without being defined in the uplugin descriptor
+
+$ DisableMergingModuleAndGeneratedFilesInUnityFiles (String[]): List of modules that disables merging module and generated cpp files in the same unity files.
+
+$ DisableUnityBuildForModules (String[]): List of modules to disable unity builds for
+
+$ EnableOptimizeCodeForModules (String[]): List of modules to enable optimizations for
+
+$ DisableOptimizeCodeForModules (String[]): List of modules to disable optimizations for
+
+$ OptimizeForSizeModules (String[]): List of modules to optimize for size. This allows the target to override module optimization level Note that this may disable PCH usage if a private PCH is not provided
+
+$ OptimizeForSizeAndSpeedModules (String[]): List of modules to optimize for size and speed. This allows the target to override module optimization level Note that this may disable PCH usage if a private PCH is not provided
+
+$ OptedInModulePlatforms (UnrealTargetPlatform[]): If this is non-null, then any platforms NOT listed will not be allowed to have modules in their directories be created
+
+$ Name (String): The name of this target
+
+$ Logger (ILogger): Logger for output relating to this target. Set before the constructor is run from
+
+$ AdditionalData (Object): Generic nullable object so a user can set additional data in a project's TargetRule and access in a project's ModuleRule without needing to add new properties post-release. For example:
+        * in .Target.cs:	AdditionalData = "data";
+        * in .Build.cs:		if ((Target.AdditionalData as string) == "data") { ... }
+
+$ Platform (UnrealTargetPlatform): Platform that this target is being built for.
+
+$ Configuration (UnrealTargetConfiguration): The configuration being built.
+
+$ Architectures (UnrealArchitectures): Architecture that the target is being built for (or an empty string for the default).
+
 $ IntermediateEnvironment (UnrealIntermediateEnvironment): Intermediate environment. Determines if the intermediates end up in a different folder than normal.
 
+$ ProjectFile (FileReference): Path to the project file for the project containing this target.
+
+$ Version (ReadOnlyBuildVersion): The current build version
+
 $ Type (TargetType): The type of target.
 
+$ DefaultBuildSettings (BuildSettingsVersion): Specifies the engine version to maintain backwards-compatible default build settings with (eg. DefaultSettingsVersion.Release_4_23, DefaultSettingsVersion.Release_4_24). Specify DefaultSettingsVersion.Latest to always use defaults for the current engine version, at the risk of introducing build errors while upgrading.
+
 $ ForcedIncludeOrder (Nullable&lt;EngineIncludeOrderVersion&gt;): Force the include order to a specific version. Overrides any Target and Module rules.
 
+$ IncludeOrderVersion (EngineIncludeOrderVersion): What version of include order to use when compiling this target. Can be overridden via -ForceIncludeOrder on the command line. ModuleRules.IncludeOrderVersion takes precedence.
+
 $ OutputFile (String): Path to the output file for the main executable, relative to the Engine or project directory. This setting is only typically useful for non-UE programs, since the engine uses paths relative to the executable to find other known folders (eg. Content).
 
 $ bUsesSteam (Boolean): Whether the target uses Steam.
@@ -68,26 +124,16 @@
 
 $ UndecoratedConfiguration (UnrealTargetConfiguration): Specifies the configuration whose binaries do not require a "-Platform-Configuration" suffix.
 
+$ bAllowHotReload (Boolean): Whether this target supports hot reload
+
 $ bBuildAllModules (Boolean): Build all the modules that are valid for this target type. Used for CIS and making installed engine builds.
 
 $ VSTestRunSettingsFile (FileReference): Set this to reference a VSTest run settings file from generated projects.
 
-$ BuildPlugins (List&lt;String&gt;): Additional plugins that are built for this target type but not enabled.
-
 $ bRuntimeDependenciesComeFromBuildPlugins (Boolean): If this is true, then the BuildPlugins list will be used to populate RuntimeDependencies, rather than EnablePlugins
 
-$ AdditionalPlugins (List&lt;String&gt;): A list of additional plugins which need to be included in this target. This allows referencing non-optional plugin modules which cannot be disabled, and allows building against specific modules in program targets which do not fit the categories in ModuleHostType.
-
-$ EnablePlugins (List&lt;String&gt;): Additional plugins that should be included for this target.
-
-$ DisablePlugins (List&lt;String&gt;): List of plugins to be disabled for this target. Note that the project file may still reference them, so they should be marked as optional to avoid failing to find them at runtime.
-
-$ OptionalPlugins (List&lt;String&gt;): Additional plugins that should be included for this target if they are found.
-
 $ DisablePluginsConflictWarningLevel (WarningLevel): How to treat conflicts when a disabled plugin is being enabled by another plugin referencing it
 
-$ InternalPluginDependencies (List&lt;String&gt;): A list of Plugin names that are allowed to exist as dependencies without being defined in the uplugin descriptor
-
 $ PakSigningKeysFile (String): Path to the set of pak signing keys to embed in the executable.
 
 $ SolutionDirectory (String): Allows a Program Target to specify it's own solution folder path.
@@ -136,24 +182,52 @@
 
 $ bCompilePython (Boolean): Whether to compile in python support
 
+$ bUseGameplayDebugger (Boolean): Whether to compile with WITH_GAMEPLAY_DEBUGGER enabled with all Engine's default gameplay debugger categories.
+
+$ bUseGameplayDebuggerCore (Boolean): Set to true when bUseGameplayDebugger is false but GameplayDebugger's core parts are required.
+
 $ bUseIris (Boolean): Whether to use Iris.
 
+$ bBuildEditor (Boolean): Whether we are compiling editor code or not. Prefer the more explicit bCompileAgainstEditor instead.
+
+$ bBuildRequiresCookedData (Boolean): Whether to compile code related to building assets. Consoles generally cannot build assets. Desktop platforms generally can.
+
+$ bBuildWithEditorOnlyData (Boolean): Whether to compile WITH_EDITORONLY_DATA disabled. Only Windows will use this, other platforms force this to false.
+
+$ bBuildDeveloperTools (Boolean): Whether to compile the developer tools.
+
+$ bBuildTargetDeveloperTools (Boolean): Whether to compile the developer tools that are for target platforms or connected devices (defaults to bBuildDeveloperTools)
+
 $ bForceBuildTargetPlatforms (Boolean): Whether to force compiling the target platform modules, even if they wouldn't normally be built.
 
 $ bForceBuildShaderFormats (Boolean): Whether to force compiling shader format modules, even if they wouldn't normally be built.
 
 $ bNeedsExtraShaderFormatsOverride (Nullable&lt;Boolean&gt;): Override for including extra shader formats
 
+$ bNeedsExtraShaderFormats (Boolean): Whether we should include any extra shader formats. By default this is only enabled for Program and Editor targets.
+
 $ bCompileCustomSQLitePlatform (Boolean): Whether we should compile SQLite using the custom "Unreal" platform (true), or using the native platform (false).
 
 $ bUseCacheFreedOSAllocs (Boolean): Whether to utilize cache freed OS allocs with MallocBinned
 
+$ bCompileAgainstEngine (Boolean): Enabled for all builds that include the engine project.  Disabled only when building standalone apps that only link with Core.
+
+$ bCompileAgainstCoreUObject (Boolean): Enabled for all builds that include the CoreUObject project.  Disabled only when building standalone apps that only link with Core.
+
+$ bCompileAgainstApplicationCore (Boolean): Enabled for builds that need to initialize the ApplicationCore module. Command line utilities do not normally need this.
+
+$ bEnableTrace (Boolean): Enabled for builds that want to use against the Trace module for profiling and diagnostics.
+
+$ bCompileAgainstEditor (Boolean): Enabled for editor builds (TargetType.Editor). Can be overridden for programs (TargetType.Program) that would need to compile against editor code. Not available for other target types. Mainly drives the value of WITH_EDITOR.
+
 $ bCompileRecast (Boolean): Whether to compile Recast navmesh generation.
 
 $ bCompileNavmeshSegmentLinks (Boolean): Whether to compile with navmesh segment links.
 
 $ bCompileNavmeshClusterLinks (Boolean): Whether to compile with navmesh cluster links.
 
+$ bCompileSpeedTree (Boolean): Whether we should compile in support for SpeedTree or not.
+
 $ bForceEnableExceptions (Boolean): Enable exceptions for all modules.
 
 $ bUseInlining (Boolean): Enable inlining for all modules.
@@ -166,12 +240,24 @@
 
 $ bEnableStackProtection (Boolean): Enable Stack Protection. Has an overhead cost
 
+$ bWithServerCode (Boolean): Compile server-only code.
+
+$ bFNameOutlineNumber (Boolean): Compile with FName storing the number part in the name table. Saves memory when most names are not numbered and those that are are referenced multiple times. The game and engine must ensure they reuse numbered names similarly to name strings to avoid leaking memory.
+
+$ bWithPushModel (Boolean): When enabled, Push Model Networking support will be compiled in. This can help reduce CPU overhead of networking, at the cost of more memory. Always enabled in editor builds.
+
 $ bCompileWithStatsWithoutEngine (Boolean): Whether to include stats support even without the engine.
 
 $ bCompileWithPluginSupport (Boolean): Whether to include plugin support.
 
+$ bIncludePluginsForTargetPlatforms (Boolean): Whether to allow plugins which support all target platforms.
+
 $ bCompileWithAccessibilitySupport (Boolean): Whether to allow accessibility code in both Slate and the OS layer.
 
+$ bWithPerfCounters (Boolean): Whether to include PerfCounters support.
+
+$ bWithLiveCoding (Boolean): Whether to enable support for live coding
+
 $ bUseDebugLiveCodingConsole (Boolean): Whether to enable support for live coding
 
 $ bWithDirectXMath (Boolean): Whether to enable support for DirectX Math LWC_TODO: No longer supported. Needs removing.
@@ -218,32 +304,42 @@
 
 $ bUseXGEController (Boolean): Whether the XGE controller worker and modules should be included in the engine build. These are required for distributed shader compilation using the XGE interception interface.
 
+$ bIncludeHeaders (Boolean): Adds header files in included modules to the build.
+
 $ bHeadersOnly (Boolean): When used with -IncludeHeaders, only header files will be compiled.
 
 $ bEnforceIWYU (Boolean): Enforce "include what you use" rules; warns if monolithic headers (Engine.h, UnrealEd.h, etc...) are used, and checks that source files include their matching header first.
 
 $ bWarnAboutMonolithicHeadersIncluded (Boolean): Emit a warning when an old-style monolithic header is included while compiling this target.
 
+$ bHasExports (Boolean): Whether the final executable should export symbols.
+
 $ bPrecompile (Boolean): Make static libraries for all engine modules as intermediates for this target.
 
 $ bEnableOSX109Support (Boolean): Whether we should compile with support for OS X 10.9 Mavericks. Used for some tools that we need to be compatible with this version of OS X.
 
 $ bIsBuildingConsoleApplication (Boolean): True if this is a console application that's being built.
 
+$ bBuildAdditionalConsoleApp (Boolean): If true, creates an additional console application. Hack for Windows, where it's not possible to conditionally inherit a parent's console Window depending on how the application is invoked; you have to link the same executable with a different subsystem setting.
+
 $ bDisableSymbolCache (Boolean): True if debug symbols that are cached for some platforms should not be created.
 
+$ bUseUnityBuild (Boolean): Whether to unify C++ code into larger files for faster compilation.
+
 $ bForceUnityBuild (Boolean): Whether to force C++ source files to be combined into larger files for faster compilation.
 
 $ bMergeModuleAndGeneratedUnityFiles (Boolean): Whether to merge module and generated unity files for faster compilation.
 
-$ DisableMergingModuleAndGeneratedFilesInUnityFiles (String[]): List of modules that disables merging module and generated cpp files in the same unity files.
-
 $ bUseAdaptiveUnityBuild (Boolean): Use a heuristic to determine which files are currently being iterated on and exclude them from unity blobs, result in faster incremental compile times. The current implementation uses the read-only flag to distinguish the working set, assuming that files will be made writable by the source control system if they are being modified. This is true for Perforce, but not for Git.
 
 $ bAdaptiveUnityDisablesOptimizations (Boolean): Disable optimization for files that are in the adaptive non-unity working set.
 
 $ bAdaptiveUnityDisablesPCH (Boolean): Disables force-included PCHs for files that are in the adaptive non-unity working set.
 
+$ bAdaptiveUnityDisablesProjectPCHForProjectPrivate (Nullable&lt;Boolean&gt;): Backing storage for bAdaptiveUnityDisablesProjectPCH.
+
+$ bAdaptiveUnityDisablesPCHForProject (Boolean): Whether to disable force-included PCHs for project source files in the adaptive non-unity working set. Defaults to bAdaptiveUnityDisablesPCH;
+
 $ bAdaptiveUnityCreatesDedicatedPCH (Boolean): Creates a dedicated PCH for each source file in the working set, allowing faster iteration on cpp-only changes.
 
 $ bAdaptiveUnityEnablesEditAndContinue (Boolean): Creates a dedicated PCH for each source file in the working set, allowing faster iteration on cpp-only changes.
@@ -252,6 +348,14 @@
 
 $ MinGameModuleSourceFilesForUnityBuild (Int32): The number of source files in a game module before unity build will be activated for that module.  This allows small game modules to have faster iterative compile times for single files, at the expense of slower full rebuild times.  This setting can be overridden by the bFasterWithoutUnity option in a module's Build.cs file.
 
+$ DefaultWarningLevel (WarningLevel): Default treatment of uncategorized warnings
+
+$ bRequireObjectPtrForAddReferencedObjects (Boolean): Require TObjectPtr for FReferenceCollector API's. (Needed for compatibility with incremental GC.)
+
+$ bValidateFormatStrings (Boolean): Emits compilation errors for incorrect UE_LOG format strings.
+
+$ DeprecationWarningLevel (WarningLevel): Level to report deprecation warnings as errors
+
 $ ShadowVariableWarningLevel (WarningLevel): Forces shadow variable warnings to be treated as errors on platforms that support it.
 
 $ bWarningsAsErrors (Boolean): Whether to enable all warnings as errors. UE enables most warnings as errors already, but disables a few (such as deprecation warnings).
@@ -268,6 +372,8 @@
 
 $ ModuleIncludeSubdirectoryWarningLevel (WarningLevel): How to treat unnecessary module sub-directory include path validation messages
 
+$ bRetainFramePointers (Boolean): Forces frame pointers to be retained this is usually required when you want reliable callstacks e.g. mallocframeprofiler
+
 $ bRetainFramePointersOverride (Nullable&lt;Boolean&gt;): Forces frame pointers to be retained this is usually required when you want reliable callstacks e.g. mallocframeprofiler
 
 $ bUseFastMonoCalls (Boolean): New Monolithic Graphics drivers have optional "fast calls" replacing various D3d functions
@@ -282,6 +388,12 @@
 
 $ bForceDebugInfo (Boolean): Whether to force debug info to be generated.
 
+$ bDisableDebugInfo (Boolean): Whether to globally disable debug info generation; Obsolete, please use TargetRules.DebugInfoMode instead
+
+$ DebugInfo (DebugInfoMode): How much debug info should be generated. See DebugInfoMode enum for more details
+
+$ DebugInfoLineTablesOnly (DebugInfoMode): True if only debug line number tables should be emitted in debug information for compilers that support doing so. Overrides TargetRules.DebugInfo See https://clang.llvm.org/docs/UsersManual.html#cmdoption-gline-tables-only for more information
+
 $ bDisableDebugInfoForGeneratedCode (Boolean): Whether to disable debug info generation for generated files. This improves link times and reduces pdb size for modules that have a lot of generated glue code.
 
 $ bOmitPCDebugInfoInDevelopment (Boolean): Whether to disable debug info on PC/Mac in development builds (for faster developer iteration, as link times are extremely fast with debug info disabled).
@@ -290,6 +402,8 @@
 
 $ bUsePCHFiles (Boolean): Whether PCH files should be used.
 
+$ bDeterministic (Boolean): Set flags require for deterministic compiling and linking. Enabling deterministic mode for msvc disables codegen multithreading so compiling will be slower
+
 $ bChainPCHs (Boolean): Whether PCHs should be chained when compiling with clang.
 
 $ bForceIncludePCHHeadersForGenCppFilesWhenPCHIsDisabled (Boolean): Whether PCH headers should be force included for gen.cpp files when PCH is disabled.
@@ -398,148 +512,36 @@
 
 $ bAllowNonUFSIniWhenCooked (Boolean): Whether to load non-ufs ini files in cooked build, (GameUserSettings.ini loaded either way)
 
+$ bLegacyPublicIncludePaths (Boolean): Add all the public folders as include paths for the compile environment.
+
+$ bLegacyParentIncludePaths (Boolean): Add all the parent folders as include paths for the compile environment.
+
+$ CppStandardEngine (CppStandardVersion): Which C++ standard to use for compiling this target (for engine modules)
+
+$ CppStandard (CppStandardVersion): Which C++ standard to use for compiling this target (for non-engine modules)
+
 $ CStandard (CStandardVersion): Which C standard to use for compiling this target
 
 $ MinCpuArchX64 (MinimumCpuArchitectureX64): Direct the compiler to generate AVX instructions wherever SSE or AVX intrinsics are used, on the x64 platforms that support it. Note that by enabling this you are changing the minspec for the PC platform, and the resultant executable will crash on machines without AVX support.
 
 $ BuildVersion (String): The build version string
 
-$ GlobalDefinitions (List&lt;String&gt;): Macros to define globally across the whole target.
+$ LinkType (TargetLinkType): Specifies how to link modules in this target (monolithic or modular). This is currently protected for backwards compatibility. Call the GetLinkType() accessor until support for the deprecated ShouldCompileMonolithic() override has been removed.
 
-$ ProjectDefinitions (List&lt;String&gt;): Macros to define across all macros in the project.
+$ LaunchModuleName (String): Specifies the name of the launch module. For modular builds, this is the module that is compiled into the target's executable.
 
 $ ExportPublicHeader (String): Specifies the path to write a header containing public definitions for this target. Useful when building a DLL to be consumed by external build processes.
 
-$ ExtraModuleNames (List&lt;String&gt;): List of additional modules to be compiled into the target.
+$ BuildEnvironment (TargetBuildEnvironment): Specifies the build environment for this target. See TargetBuildEnvironment for more information on the available options.
 
-$ ManifestFileNames (List&lt;FileReference&gt;): Path to a manifest to output for this target
-
-$ DependencyListFileNames (List&lt;FileReference&gt;): Path to a list of dependencies for this target, when precompiling
-
 $ bOverrideBuildEnvironment (Boolean): Whether to ignore violations to the shared build environment (eg. editor targets modifying definitions)
 
-$ PreBuildTargets (List&lt;TargetInfo&gt;): Specifies a list of targets which should be built before this target is built.
-
-$ PreBuildSteps (List&lt;String&gt;): Specifies a list of steps which should be executed before this target is built, in the context of the host platform's shell. The following variables will be expanded before execution: $(EngineDir), $(ProjectDir), $(TargetName), $(TargetPlatform), $(TargetConfiguration), $(TargetType), $(ProjectFile).
-
-$ PostBuildSteps (List&lt;String&gt;): Specifies a list of steps which should be executed after this target is built, in the context of the host platform's shell. The following variables will be expanded before execution: $(EngineDir), $(ProjectDir), $(TargetName), $(TargetPlatform), $(TargetConfiguration), $(TargetType), $(ProjectFile).
-
-$ AdditionalBuildProducts (List&lt;String&gt;): Specifies additional build products produced as part of this target.
-
 $ AdditionalCompilerArguments (String): Additional arguments to pass to the compiler
 
 $ AdditionalLinkerArguments (String): Additional arguments to pass to the linker
 
 $ MemoryPerActionGB (Double): Max amount of memory that each compile action may require. Used by ParallelExecutor to decide the maximum number of parallel actions to start at one time.
 
-$ DisableUnityBuildForModules (String[]): List of modules to disable unity builds for
-
-$ EnableOptimizeCodeForModules (String[]): List of modules to enable optimizations for
-
-$ DisableOptimizeCodeForModules (String[]): List of modules to disable optimizations for
-
-$ OptimizeForSizeModules (String[]): List of modules to optimize for size. This allows the target to override module optimization level Note that this may disable PCH usage if a private PCH is not provided
-
-$ OptimizeForSizeAndSpeedModules (String[]): List of modules to optimize for size and speed. This allows the target to override module optimization level Note that this may disable PCH usage if a private PCH is not provided
-
 $ GeneratedProjectName (String): When generating project files, specifies the name of the project file to use when there are multiple targets of the same type.
 
-$ OptedInModulePlatforms (UnrealTargetPlatform[]): If this is non-null, then any platforms NOT listed will not be allowed to have modules in their directories be created
 
-$ Name (String): The name of this target
-
-$ Logger (ILogger): Logger for output relating to this target. Set before the constructor is run from
-
-$ AdditionalData (Object): Generic nullable object so a user can set additional data in a project's TargetRule and access in a project's ModuleRule without needing to add new properties post-release. For example:
-        * in .Target.cs:	AdditionalData = "data";
-        * in .Build.cs:		if ((Target.AdditionalData as string) == "data") { ... }
-
-$ DefaultBuildSettings (BuildSettingsVersion): Specifies the engine version to maintain backwards-compatible default build settings with (eg. DefaultSettingsVersion.Release_4_23, DefaultSettingsVersion.Release_4_24). Specify DefaultSettingsVersion.Latest to always use defaults for the current engine version, at the risk of introducing build errors while upgrading.
-
-$ IncludeOrderVersion (EngineIncludeOrderVersion): What version of include order to use when compiling this target. Can be overridden via -ForceIncludeOrder on the command line. ModuleRules.IncludeOrderVersion takes precedence.
-
-$ bAllowHotReload (Boolean): Whether this target supports hot reload
-
-$ bUseGameplayDebugger (Boolean): Whether to compile with WITH_GAMEPLAY_DEBUGGER enabled with all Engine's default gameplay debugger categories.
-
-$ bUseGameplayDebuggerCore (Boolean): Set to true when bUseGameplayDebugger is false but GameplayDebugger's core parts are required.
-
-$ bBuildEditor (Boolean): Whether we are compiling editor code or not. Prefer the more explicit bCompileAgainstEditor instead.
-
-$ bBuildRequiresCookedData (Boolean): Whether to compile code related to building assets. Consoles generally cannot build assets. Desktop platforms generally can.
-
-$ bBuildWithEditorOnlyData (Boolean): Whether to compile WITH_EDITORONLY_DATA disabled. Only Windows will use this, other platforms force this to false.
-
-$ bBuildDeveloperTools (Boolean): Whether to compile the developer tools.
-
-$ bBuildTargetDeveloperTools (Boolean): Whether to compile the developer tools that are for target platforms or connected devices (defaults to bBuildDeveloperTools)
-
-$ bNeedsExtraShaderFormats (Boolean): Whether we should include any extra shader formats. By default this is only enabled for Program and Editor targets.
-
-$ bCompileAgainstEngine (Boolean): Enabled for all builds that include the engine project.  Disabled only when building standalone apps that only link with Core.
-
-$ bCompileAgainstCoreUObject (Boolean): Enabled for all builds that include the CoreUObject project.  Disabled only when building standalone apps that only link with Core.
-
-$ bCompileAgainstApplicationCore (Boolean): Enabled for builds that need to initialize the ApplicationCore module. Command line utilities do not normally need this.
-
-$ bEnableTrace (Boolean): Enabled for builds that want to use against the Trace module for profiling and diagnostics.
-
-$ bCompileAgainstEditor (Boolean): Enabled for editor builds (TargetType.Editor). Can be overridden for programs (TargetType.Program) that would need to compile against editor code. Not available for other target types. Mainly drives the value of WITH_EDITOR.
-
-$ bCompileSpeedTree (Boolean): Whether we should compile in support for SpeedTree or not.
-
-$ bWithServerCode (Boolean): Compile server-only code.
-
-$ bFNameOutlineNumber (Boolean): Compile with FName storing the number part in the name table. Saves memory when most names are not numbered and those that are are referenced multiple times. The game and engine must ensure they reuse numbered names similarly to name strings to avoid leaking memory.
-
-$ bWithPushModel (Boolean): When enabled, Push Model Networking support will be compiled in. This can help reduce CPU overhead of networking, at the cost of more memory. Always enabled in editor builds.
-
-$ bIncludePluginsForTargetPlatforms (Boolean): Whether to allow plugins which support all target platforms.
-
-$ bWithPerfCounters (Boolean): Whether to include PerfCounters support.
-
-$ bWithLiveCoding (Boolean): Whether to enable support for live coding
-
-$ bIncludeHeaders (Boolean): Adds header files in included modules to the build.
-
-$ bHasExports (Boolean): Whether the final executable should export symbols.
-
-$ bBuildAdditionalConsoleApp (Boolean): If true, creates an additional console application. Hack for Windows, where it's not possible to conditionally inherit a parent's console Window depending on how the application is invoked; you have to link the same executable with a different subsystem setting.
-
-$ bUseUnityBuild (Boolean): Whether to unify C++ code into larger files for faster compilation.
-
-$ bAdaptiveUnityDisablesPCHForProject (Boolean): Whether to disable force-included PCHs for project source files in the adaptive non-unity working set. Defaults to bAdaptiveUnityDisablesPCH;
-
-$ DefaultWarningLevel (WarningLevel): Default treatment of uncategorized warnings
-
-$ bRequireObjectPtrForAddReferencedObjects (Boolean): Require TObjectPtr for FReferenceCollector API's. (Needed for compatibility with incremental GC.)
-
-$ bValidateFormatStrings (Boolean): Emits compilation errors for incorrect UE_LOG format strings.
-
-$ DeprecationWarningLevel (WarningLevel): Level to report deprecation warnings as errors
-
-$ bRetainFramePointers (Boolean): Forces frame pointers to be retained this is usually required when you want reliable callstacks e.g. mallocframeprofiler
-
-$ bDisableDebugInfo (Boolean): Whether to globally disable debug info generation; Obsolete, please use TargetRules.DebugInfoMode instead
-
-$ DebugInfo (DebugInfoMode): How much debug info should be generated. See DebugInfoMode enum for more details
-
-$ DebugInfoLineTablesOnly (DebugInfoMode): True if only debug line number tables should be emitted in debug information for compilers that support doing so. Overrides TargetRules.DebugInfo See https://clang.llvm.org/docs/UsersManual.html#cmdoption-gline-tables-only for more information
-
-$ bDeterministic (Boolean): Set flags require for deterministic compiling and linking. Enabling deterministic mode for msvc disables codegen multithreading so compiling will be slower
-
-$ bLegacyPublicIncludePaths (Boolean): Add all the public folders as include paths for the compile environment.
-
-$ bLegacyParentIncludePaths (Boolean): Add all the parent folders as include paths for the compile environment.
-
-$ CppStandardEngine (CppStandardVersion): Which C++ standard to use for compiling this target (for engine modules)
-
-$ CppStandard (CppStandardVersion): Which C++ standard to use for compiling this target (for non-engine modules)
-
-$ LinkType (TargetLinkType): Specifies how to link modules in this target (monolithic or modular). This is currently protected for backwards compatibility. Call the GetLinkType() accessor until support for the deprecated ShouldCompileMonolithic() override has been removed.
-
-$ LaunchModuleName (String): Specifies the name of the launch module. For modular builds, this is the module that is compiled into the target's executable.
-
-$ BuildEnvironment (TargetBuildEnvironment): Specifies the build environment for this target. See TargetBuildEnvironment for more information on the available options.
-
-
