Documentation for build tools

diff --git a/Engine/Documentation/Source/Programming/Development/BuildGraph/ScriptAnatomy/Tasks/PredefinedTasks/BuildGraphScriptPredefinedTasks.INT.udn b/Engine/Documentation/Source/Programming/Development/BuildGraph/ScriptAnatomy/Tasks/PredefinedTasks/BuildGraphScriptPredefinedTasks.INT.udn
--- a/Engine/Documentation/Source/Programming/Development/BuildGraph/ScriptAnatomy/Tasks/PredefinedTasks/BuildGraphScriptPredefinedTasks.INT.udn
+++ b/Engine/Documentation/Source/Programming/Development/BuildGraph/ScriptAnatomy/Tasks/PredefinedTasks/BuildGraphScriptPredefinedTasks.INT.udn
@@ -419,6 +419,8 @@ Spawns Docker and waits for it to complete.
 | BaseDir        | String             | Required | Base directory for running the command                                                                                                                   |
 | IgnorePackages | String             | Optional | Specifies a list of packages to ignore for version checks, separated by semicolons. Optional version number may be specified with 'name@version' syntax. |
 | LicenseDir     | DirectoryReference | Optional | Directory containing allowed licenses                                                                                                                    |
+| CsvFile        | FileReference      | Optional | Path to a csv file to write with list of packages and their licenses                                                                                     |
+| DotNetPath     | FileReference      | Optional | Override path to dotnet executable                                                                                                                       |
 | If             | Condition          | Optional | Whether to execute this task. It is ignored if this condition evaluates to false.                                                                        |
 
 ### OnExit
