Documentation for build tools

diff --git a/Engine/Documentation/Source/Programming/UnrealBuildSystem/Configuration/BuildConfigProperties/BuildConfigProperties.INT.udn b/Engine/Documentation/Source/Programming/UnrealBuildSystem/Configuration/BuildConfigProperties/BuildConfigProperties.INT.udn
--- a/Engine/Documentation/Source/Programming/UnrealBuildSystem/Configuration/BuildConfigProperties/BuildConfigProperties.INT.udn
+++ b/Engine/Documentation/Source/Programming/UnrealBuildSystem/Configuration/BuildConfigProperties/BuildConfigProperties.INT.udn
@@ -458,7 +458,13 @@ $ Token : Uri of the Horde server
 
 $ OidcProvider : OIDC id for the login to use
 
-$ Pool : Pool for the Horde agent to assign
+$ Pool : Pool for the Horde agent to assign if no override current platform doesn't have it set
+
+$ LinuxPool : Pool for the Horde agent to assign when on Linux
+
+$ MacPool : Pool for the Horde agent to assign when on Mac
+
+$ WindowsPool : Pool for the Horde agent to assign when on Windows
 
 $ Requirements : Requirements for the Horde agent to assign
 
