Documentation for build tools

diff --git a/Engine/Documentation/Source/Programming/Development/BuildGraph/ScriptAnatomy/Tasks/PredefinedTasks/BuildGraphScriptPredefinedTasks.INT.udn b/Engine/Documentation/Source/Programming/Development/BuildGraph/ScriptAnatomy/Tasks/PredefinedTasks/BuildGraphScriptPredefinedTasks.INT.udn
--- a/Engine/Documentation/Source/Programming/Development/BuildGraph/ScriptAnatomy/Tasks/PredefinedTasks/BuildGraphScriptPredefinedTasks.INT.udn
+++ b/Engine/Documentation/Source/Programming/Development/BuildGraph/ScriptAnatomy/Tasks/PredefinedTasks/BuildGraphScriptPredefinedTasks.INT.udn
@@ -518,11 +518,12 @@ Updates the local version files (Engine/Source/Runtime/Launch/Resources/Version.
 
 Signs a set of executable files with an installed certificate.
 
-|       |           |          |                                                                                                                      |
-| ----- | --------- | -------- | -------------------------------------------------------------------------------------------------------------------- |
-| Files | FileSpec  | Required | List of file specifications separated by semicolons (for example, *.cpp;Engine/.../*.bat), or the name of a tag set. |
-| Tag   | TagList   | Optional | Tag to be applied to build products of this task.                                                                    |
-| If    | Condition | Optional | Whether to execute this task. It is ignored if this condition evaluates to false.                                    |
+|             |           |          |                                                                                                                      |
+| ----------- | --------- | -------- | -------------------------------------------------------------------------------------------------------------------- |
+| Files       | FileSpec  | Required | List of file specifications separated by semicolons (for example, *.cpp;Engine/.../*.bat), or the name of a tag set. |
+| Description | String    | Optional | Optional description for the signed content                                                                          |
+| Tag         | TagList   | Optional | Tag to be applied to build products of this task.                                                                    |
+| If          | Condition | Optional | Whether to execute this task. It is ignored if this condition evaluates to false.                                    |
 
 ### Spawn
 
