Documentation for build tools

diff --git a/Engine/Documentation/Source/Programming/UnrealBuildSystem/Configuration/BuildConfigProperties/BuildConfigProperties.INT.udn b/Engine/Documentation/Source/Programming/UnrealBuildSystem/Configuration/BuildConfigProperties/BuildConfigProperties.INT.udn
--- a/Engine/Documentation/Source/Programming/UnrealBuildSystem/Configuration/BuildConfigProperties/BuildConfigProperties.INT.udn
+++ b/Engine/Documentation/Source/Programming/UnrealBuildSystem/Configuration/BuildConfigProperties/BuildConfigProperties.INT.udn
@@ -870,6 +870,8 @@ $ bBuildUBTInDebug : Forces UBT to be built in debug configuration, regardless o
 
 $ bBuildLiveCodingConsole : Whether to include a dependency on LiveCodingConsole when building targets that support live coding.
 
+$ bMakeProjectPerTarget : Whether to generate a project file for each individual target, and not include e.g. Editor/Client/Server in the Configuration.
+
 $ bIncludeDocumentation : True if we should include documentation in the generated projects.
 
 $ bUsePrecompiled : True if build targets should pass the -useprecompiled argument.
