Documentation for build tools

diff --git a/Engine/Documentation/Source/Programming/Development/BuildGraph/ScriptAnatomy/Tasks/PredefinedTasks/BuildGraphScriptPredefinedTasks.INT.udn b/Engine/Documentation/Source/Programming/Development/BuildGraph/ScriptAnatomy/Tasks/PredefinedTasks/BuildGraphScriptPredefinedTasks.INT.udn
--- a/Engine/Documentation/Source/Programming/Development/BuildGraph/ScriptAnatomy/Tasks/PredefinedTasks/BuildGraphScriptPredefinedTasks.INT.udn
+++ b/Engine/Documentation/Source/Programming/Development/BuildGraph/ScriptAnatomy/Tasks/PredefinedTasks/BuildGraphScriptPredefinedTasks.INT.udn
@@ -60,6 +60,15 @@ Creates a new AWS ECS task definition and updates the ECS service to use this ne
 | LogOutput          | Boolean   | Optional | Write output to the log                                                                      |
 | If                 | Condition | Optional | Whether to execute this task. It is ignored if this condition evaluates to false.            |
 
+### CheckMarkdown
+
+Checks that all markdown links between the given files are valid.
+
+|       |           |          |                                                                                   |
+| ----- | --------- | -------- | --------------------------------------------------------------------------------- |
+| Files | FileSpec  | Optional | Optional filter to be applied to the list of input files.                         |
+| If    | Condition | Optional | Whether to execute this task. It is ignored if this condition evaluates to false. |
+
 ### Command
 
 Invokes an AutomationTool child process to run the given command.
diff --git a/Engine/Documentation/Source/Programming/UnrealBuildSystem/Configuration/BuildConfigProperties/BuildConfigProperties.INT.udn b/Engine/Documentation/Source/Programming/UnrealBuildSystem/Configuration/BuildConfigProperties/BuildConfigProperties.INT.udn
--- a/Engine/Documentation/Source/Programming/UnrealBuildSystem/Configuration/BuildConfigProperties/BuildConfigProperties.INT.udn
+++ b/Engine/Documentation/Source/Programming/UnrealBuildSystem/Configuration/BuildConfigProperties/BuildConfigProperties.INT.udn
@@ -338,6 +338,8 @@ $ InlineFunctionExpansionLevel : Which level to use for Inline Function Expansio
 
 $ ToolChain : Version of the toolchain to use on Windows platform when a non-msvc Compiler is in use, to locate include paths etc.
 
+$ ToolchainVersionWarningLevel : Warning level when reporting toolchains that are not in the preferred version list
+
 $ bStrictConformanceMode : Enables strict standard conformance mode (/permissive-).
 
 ### ModuleConfiguration
@@ -856,6 +858,32 @@ $ PrimaryProjectName : Name of the primary project file -- for example, the base
 
 $ bPrimaryProjectNameFromFolder : If true, sets the primary project name according to the name of the folder it is in.
 
+### VSWorkspaceProjectFileGenerator
+
+$ bIncludeDocumentation : True if we should include documentation in the generated projects.
+
+$ bUsePrecompiled : True if build targets should pass the -useprecompiled argument.
+
+$ bIncludeEngineSource : True if we should include engine source in the generated solution.
+
+$ bIncludeShaderSource : True if shader source files should be included in generated projects.
+
+$ bIncludeConfigFiles : True if we should include config (.ini) files in the generated project.
+
+$ bIncludeTemplateFiles : True if we should include template files in the generated project.
+
+$ bIncludeEnginePrograms : True if we should include program projects in the generated solution.
+
+$ IncludeCppSource : Whether to include C++ targets
+
+$ bIncludeDotNetPrograms : True if we should include csharp program projects in the generated solution. Pass "-DotNet" to enable this.
+
+$ bIncludeTempTargets : Whether to include temporary targets generated by UAT to support content only projects with non-default settings.
+
+$ PrimaryProjectName : Name of the primary project file -- for example, the base file name for the Visual Studio solution file, or the Xcode project file on Mac.
+
+$ bPrimaryProjectNameFromFolder : If true, sets the primary project name according to the name of the folder it is in.
+
 ### VCProjectFileGenerator
 
 $ Version : The version of Visual Studio to generate project files for.
