Documentation for build tools

diff --git a/Engine/Documentation/Source/Programming/Development/BuildGraph/ScriptAnatomy/Tasks/PredefinedTasks/BuildGraphScriptPredefinedTasks.INT.udn b/Engine/Documentation/Source/Programming/Development/BuildGraph/ScriptAnatomy/Tasks/PredefinedTasks/BuildGraphScriptPredefinedTasks.INT.udn
--- a/Engine/Documentation/Source/Programming/Development/BuildGraph/ScriptAnatomy/Tasks/PredefinedTasks/BuildGraphScriptPredefinedTasks.INT.udn
+++ b/Engine/Documentation/Source/Programming/Development/BuildGraph/ScriptAnatomy/Tasks/PredefinedTasks/BuildGraphScriptPredefinedTasks.INT.udn
@@ -2,7 +2,7 @@ Availability: NoPublish
 Title: BuildGraph Predefined Tasks
 Crumbs: %ROOT%, Programming, Programming/Development, Programming/Development/BuildGraph, Programming/Development/BuildGraph/BuildGraphScriptTasks
 Description: This is a procedurally generated markdown page.
-version: 5.5
+version: 5.4
 parent:Programming/Development/BuildGraph/BuildGraphScriptTasks
 
 ### AgeStore
@@ -60,15 +60,6 @@ Creates a new AWS ECS task definition and updates the ECS service to use this ne
 | LogOutput          | Boolean   | Optional | Write output to the log                                                                      |
 | If                 | Condition | Optional | Whether to execute this task. It is ignored if this condition evaluates to false.            |
 
-### CheckMarkdown
-
-Checks that all markdown links between the given files are valid.
-
-|       |           |          |                                                                                   |
-| ----- | --------- | -------- | --------------------------------------------------------------------------------- |
-| Files | FileSpec  | Optional | Optional filter to be applied to the list of input files.                         |
-| If    | Condition | Optional | Whether to execute this task. It is ignored if this condition evaluates to false. |
-
 ### Command
 
 Invokes an AutomationTool child process to run the given command.
@@ -712,21 +703,21 @@ Writes text to a file.
 
 Exports an snapshot from Zen to a specified destination.
 
-|                            |                    |          |                                                                                                                              |
-| -------------------------- | ------------------ | -------- | ---------------------------------------------------------------------------------------------------------------------------- |
-| Project                    | FileReference      | Optional | The project from which to export the snapshot                                                                                |
-| Platform                   | String             | Optional | The target platform(s) to export the snapshot for                                                                            |
-| SnapshotBaseDescriptorFile | FileReference      | Optional | A file to read with information about the snapshot that should be used as a base when exporting this new snapshot            |
-| SnapshotDescriptorFile     | FileReference      | Optional | A file to create with information about the snapshot that was exported                                                       |
-| DestinationStorageType     | String             | Required | The type of destination to export the snapshot to (cloud, ...)                                                               |
-| DestinationIdentifier      | String             | Optional | The identifier to use when exporting to a destination                                                                        |
-| DestinationCloudHost       | String             | Optional | The host name to use when exporting to a cloud destination                                                                   |
-| DestinationCloudNamespace  | String             | Optional | The namespace to use when exporting to a cloud destination                                                                   |
-| DestinationCloudBucket     | String             | Optional | A custom bucket name to use when exporting to a cloud destination                                                            |
-| DestinationFileDir         | DirectoryReference | Optional | The directory to use when exporting to a file destination                                                                    |
-| DestinationFileName        | String             | Optional | The filename to use when exporting to a file destination                                                                     |
-| OverridePlatformCookedDir  | String             | Optional | Optional. Where to look for the ue.projectstore The pattern {Platform} can be used for exporting multiple platforms at once. |
-| If                         | Condition          | Optional | Whether to execute this task. It is ignored if this condition evaluates to false.                                            |
+|                            |                    |          |                                                                                                                            |
+| -------------------------- | ------------------ | -------- | -------------------------------------------------------------------------------------------------------------------------- |
+| Project                    | FileReference      | Optional | The project from which to export the snapshot                                                                              |
+| Platform                   | String             | Optional | The target platform(s) to export the snapshot for                                                                          |
+| SnapshotBaseDescriptorFile | FileReference      | Optional | A file to read with information about the snapshot that should be used as a base when exporting this new snapshot          |
+| SnapshotDescriptorFile     | FileReference      | Optional | A file to create with information about the snapshot that was exported                                                     |
+| DestinationStorageType     | String             | Required | The type of destination to export the snapshot to (cloud, ...)                                                             |
+| DestinationIdentifier      | String             | Optional | The identifier to use when exporting to a destination                                                                      |
+| DestinationCloudHost       | String             | Optional | The host name to use when exporting to a cloud destination                                                                 |
+| DestinationCloudNamespace  | String             | Optional | The namespace to use when exporting to a cloud destination                                                                 |
+| DestinationCloudBucket     | String             | Optional | A custom bucket name to use when exporting to a cloud destination                                                          |
+| DestinationFileDir         | DirectoryReference | Optional | The directory to use when exporting to a file destination                                                                  |
+| DestinationFileName        | String             | Optional | The filename to use when exporting to a file destination                                                                   |
+| OverridePlatformCookedDir  | String             | Optional | Optional. Where to look for the .projectstore The pattern {Platform} can be used for exporting multiple platforms at once. |
+| If                         | Condition          | Optional | Whether to execute this task. It is ignored if this condition evaluates to false.                                          |
 
 ### ZenImportOplog
 
diff --git a/Engine/Documentation/Source/Programming/UnrealBuildSystem/Configuration/BuildConfigProperties/BuildConfigProperties.INT.udn b/Engine/Documentation/Source/Programming/UnrealBuildSystem/Configuration/BuildConfigProperties/BuildConfigProperties.INT.udn
--- a/Engine/Documentation/Source/Programming/UnrealBuildSystem/Configuration/BuildConfigProperties/BuildConfigProperties.INT.udn
+++ b/Engine/Documentation/Source/Programming/UnrealBuildSystem/Configuration/BuildConfigProperties/BuildConfigProperties.INT.udn
@@ -2,7 +2,7 @@ Availability: NoPublish
 Title: Build Configuration Properties Page
 Crumbs:
 Description: This is a procedurally generated markdown page.
-Version: 5.5
+Version: 5.4
 
 ### BuildConfiguration
 
@@ -338,8 +338,6 @@ $ InlineFunctionExpansionLevel : Which level to use for Inline Function Expansio
 
 $ ToolChain : Version of the toolchain to use on Windows platform when a non-msvc Compiler is in use, to locate include paths etc.
 
-$ ToolchainVersionWarningLevel : Warning level when reporting toolchains that are not in the preferred version list
-
 $ bStrictConformanceMode : Enables strict standard conformance mode (/permissive-).
 
 ### ModuleConfiguration
@@ -858,32 +856,6 @@ $ PrimaryProjectName : Name of the primary project file -- for example, the base
 
 $ bPrimaryProjectNameFromFolder : If true, sets the primary project name according to the name of the folder it is in.
 
-### VSWorkspaceProjectFileGenerator
-
-$ bIncludeDocumentation : True if we should include documentation in the generated projects.
-
-$ bUsePrecompiled : True if build targets should pass the -useprecompiled argument.
-
-$ bIncludeEngineSource : True if we should include engine source in the generated solution.
-
-$ bIncludeShaderSource : True if shader source files should be included in generated projects.
-
-$ bIncludeConfigFiles : True if we should include config (.ini) files in the generated project.
-
-$ bIncludeTemplateFiles : True if we should include template files in the generated project.
-
-$ bIncludeEnginePrograms : True if we should include program projects in the generated solution.
-
-$ IncludeCppSource : Whether to include C++ targets
-
-$ bIncludeDotNetPrograms : True if we should include csharp program projects in the generated solution. Pass "-DotNet" to enable this.
-
-$ bIncludeTempTargets : Whether to include temporary targets generated by UAT to support content only projects with non-default settings.
-
-$ PrimaryProjectName : Name of the primary project file -- for example, the base file name for the Visual Studio solution file, or the Xcode project file on Mac.
-
-$ bPrimaryProjectNameFromFolder : If true, sets the primary project name according to the name of the folder it is in.
-
 ### VCProjectFileGenerator
 
 $ Version : The version of Visual Studio to generate project files for.
diff --git a/Engine/Documentation/Source/Programming/UnrealBuildSystem/ModuleFiles/ModuleFilesProperties/ModuleFilesProperties.INT.udn b/Engine/Documentation/Source/Programming/UnrealBuildSystem/ModuleFiles/ModuleFilesProperties/ModuleFilesProperties.INT.udn
--- a/Engine/Documentation/Source/Programming/UnrealBuildSystem/ModuleFiles/ModuleFilesProperties/ModuleFilesProperties.INT.udn
+++ b/Engine/Documentation/Source/Programming/UnrealBuildSystem/ModuleFiles/ModuleFilesProperties/ModuleFilesProperties.INT.udn
@@ -2,7 +2,7 @@ Availability: NoPublish
 Title: Build Configuration Properties Page
 Crumbs:
 Description: This is a procedurally generated markdown page.
-Version: 5.5
+Version: 5.4
 
 ### Read-Only Properties
 
diff --git a/Engine/Documentation/Source/Programming/UnrealBuildSystem/TargetFiles/TargetFilesProperties/TargetFilesProperties.INT.udn b/Engine/Documentation/Source/Programming/UnrealBuildSystem/TargetFiles/TargetFilesProperties/TargetFilesProperties.INT.udn
--- a/Engine/Documentation/Source/Programming/UnrealBuildSystem/TargetFiles/TargetFilesProperties/TargetFilesProperties.INT.udn
+++ b/Engine/Documentation/Source/Programming/UnrealBuildSystem/TargetFiles/TargetFilesProperties/TargetFilesProperties.INT.udn
@@ -2,7 +2,7 @@ Availability: NoPublish
 Title: Build Configuration Properties Page
 Crumbs:
 Description: This is a procedurally generated markdown page.
-Version: 5.5
+Version: 5.4
 
 ### Read-Only Properties
 
