Documentation for build tools

diff --git a/Engine/Documentation/Source/Programming/Development/BuildGraph/ScriptAnatomy/Tasks/PredefinedTasks/BuildGraphScriptPredefinedTasks.INT.udn b/Engine/Documentation/Source/Programming/Development/BuildGraph/ScriptAnatomy/Tasks/PredefinedTasks/BuildGraphScriptPredefinedTasks.INT.udn
--- a/Engine/Documentation/Source/Programming/Development/BuildGraph/ScriptAnatomy/Tasks/PredefinedTasks/BuildGraphScriptPredefinedTasks.INT.udn
+++ b/Engine/Documentation/Source/Programming/Development/BuildGraph/ScriptAnatomy/Tasks/PredefinedTasks/BuildGraphScriptPredefinedTasks.INT.udn
@@ -532,6 +532,7 @@ Signs a set of executable files with an installed certificate.
 | Files       | FileSpec  | Required | List of file specifications separated by semicolons (for example, *.cpp;Engine/.../*.bat), or the name of a tag set. |
 | Description | String    | Optional | Optional description for the signed content                                                                          |
 | Tag         | TagList   | Optional | Tag to be applied to build products of this task.                                                                    |
+| Parallel    | Boolean   | Optional | If true, the calls to the signing tool will be performed in parallel.                                                |
 | If          | Condition | Optional | Whether to execute this task. It is ignored if this condition evaluates to false.                                    |
 
 ### Spawn
