Documentation for build tools

diff --git a/Engine/Documentation/Source/Programming/UnrealBuildSystem/Configuration/BuildConfigProperties/BuildConfigProperties.INT.udn b/Engine/Documentation/Source/Programming/UnrealBuildSystem/Configuration/BuildConfigProperties/BuildConfigProperties.INT.udn
--- a/Engine/Documentation/Source/Programming/UnrealBuildSystem/Configuration/BuildConfigProperties/BuildConfigProperties.INT.udn
+++ b/Engine/Documentation/Source/Programming/UnrealBuildSystem/Configuration/BuildConfigProperties/BuildConfigProperties.INT.udn
@@ -444,6 +444,8 @@ $ bUseCrypto : Set to true to enable encryption when transfering files over the
 
 $ bUseKnownInputs : Set to true to provide known inputs to processes that are run remote. This is an experimental feature to speed up build times when ping is higher
 
+$ ActionsOutputFile : Write yaml file with all actions that are queued for build. This can be used to replay using "UbaCli.exe local file.yaml"
+
 $ HordeMaxCores : Max cores allowed to be used by build session
 
 $ HordeDelay : How long UBT should wait to ask for help. Useful in build configs where machine can delay remote work and still get same wall time results (pch dependencies etc)
