Documentation for build tools

diff --git a/Engine/Documentation/Source/Programming/UnrealBuildSystem/Configuration/BuildConfigProperties/BuildConfigProperties.INT.udn b/Engine/Documentation/Source/Programming/UnrealBuildSystem/Configuration/BuildConfigProperties/BuildConfigProperties.INT.udn
--- a/Engine/Documentation/Source/Programming/UnrealBuildSystem/Configuration/BuildConfigProperties/BuildConfigProperties.INT.udn
+++ b/Engine/Documentation/Source/Programming/UnrealBuildSystem/Configuration/BuildConfigProperties/BuildConfigProperties.INT.udn
@@ -60,6 +60,10 @@ $ DefaultWarningLevel : Default treatment of uncategorized warnings
 
 $ DeprecationWarningLevel : Level to report deprecation warnings as errors
 
+$ bStaticAnalyzerProjectOnly : Only run static analysis against project modules, skipping engine modules
+
+$ bStaticAnalyzerIncludeGenerated : When enabled, generated source files will be analyzed
+
 $ bDebugBuildsActuallyUseDebugCRT : Enables the debug C++ runtime (CRT) for debug builds. By default we always use the release runtime, since the debug version isn't particularly useful when debugging Unreal Engine projects, and linking against the debug CRT libraries forces our third party library dependencies to also be compiled using the debug CRT (and often perform more slowly). Often it can be inconvenient to require a separate copy of the debug versions of third party static libraries simply so that you can debug your program's code.
 
 $ bLegalToDistributeBinary : Whether the output from this target can be publicly distributed, even if it has dependencies on modules that are in folders with special restrictions (eg. CarefullyRedist, NotForLicensees, NoRedist).
diff --git a/Engine/Documentation/Source/Programming/UnrealBuildSystem/ModuleFiles/ModuleFilesProperties/ModuleFilesProperties.INT.udn b/Engine/Documentation/Source/Programming/UnrealBuildSystem/ModuleFiles/ModuleFilesProperties/ModuleFilesProperties.INT.udn
--- a/Engine/Documentation/Source/Programming/UnrealBuildSystem/ModuleFiles/ModuleFilesProperties/ModuleFilesProperties.INT.udn
+++ b/Engine/Documentation/Source/Programming/UnrealBuildSystem/ModuleFiles/ModuleFilesProperties/ModuleFilesProperties.INT.udn
@@ -12,6 +12,8 @@ $ Logger (ILogger): Accessor for the target logger
 
 $ HasPackageOverride (Boolean): Returns true if there has been an override type specified on this module
 
+$ SDKVersionRelevantPlatforms (List&lt;UnrealTargetPlatform&gt;): If this contains a platform and the project has had it's SDK version overriden from default, this module will be compiled as a project module, instead of shared engine module
+
 $ AutoSdkDirectory (String): The AutoSDK directory for the active host platform
 
 $ EngineDirectory (String): The current engine directory
diff --git a/Engine/Documentation/Source/Programming/UnrealBuildSystem/TargetFiles/TargetFilesProperties/TargetFilesProperties.INT.udn b/Engine/Documentation/Source/Programming/UnrealBuildSystem/TargetFiles/TargetFilesProperties/TargetFilesProperties.INT.udn
--- a/Engine/Documentation/Source/Programming/UnrealBuildSystem/TargetFiles/TargetFilesProperties/TargetFilesProperties.INT.udn
+++ b/Engine/Documentation/Source/Programming/UnrealBuildSystem/TargetFiles/TargetFilesProperties/TargetFilesProperties.INT.udn
@@ -68,6 +68,10 @@ $ InternalPluginDependencies (List&lt;String&gt;): A list of Plugin names that a
 
 $ DisableMergingModuleAndGeneratedFilesInUnityFiles (String[]): List of modules that disables merging module and generated cpp files in the same unity files.
 
+$ bStaticAnalyzerProjectOnly (Boolean): Only run static analysis against project modules, skipping engine modules
+
+$ bStaticAnalyzerIncludeGenerated (Boolean): When enabled, generated source files will be analyzed
+
 $ DisableUnityBuildForModules (String[]): List of modules to disable unity builds for
 
 $ EnableOptimizeCodeForModules (String[]): List of modules to enable optimizations for
@@ -536,6 +540,8 @@ $ ExportPublicHeader (String): Specifies the path to write a header containing p
 
 $ BuildEnvironment (TargetBuildEnvironment): Specifies the build environment for this target. See TargetBuildEnvironment for more information on the available options.
 
+$ bAllowSDKOverrideModulesWithSharedEnvironment (Boolean): If this true, a Shared build environment target will allow for any modules that are SDK version sensitive to have a project side module when an SDK is overridden. For instance, if IOSTargetPlatform, which is marked as IOS SDK version relevant, is compiled for a Target with this set to true, and that target overrides the IOS SDK, then this would compile to e.g. MyProject/Binaries/Win64/IOS/MyProject-IOSTargetPlatform.dll
+
 $ bOverrideBuildEnvironment (Boolean): Whether to ignore violations to the shared build environment (eg. editor targets modifying definitions)
 
 $ AdditionalCompilerArguments (String): Additional arguments to pass to the compiler
