Documentation for build tools

diff --git a/Engine/Documentation/Source/Programming/Development/BuildGraph/ScriptAnatomy/Tasks/PredefinedTasks/BuildGraphScriptPredefinedTasks.INT.udn b/Engine/Documentation/Source/Programming/Development/BuildGraph/ScriptAnatomy/Tasks/PredefinedTasks/BuildGraphScriptPredefinedTasks.INT.udn
--- a/Engine/Documentation/Source/Programming/Development/BuildGraph/ScriptAnatomy/Tasks/PredefinedTasks/BuildGraphScriptPredefinedTasks.INT.udn
+++ b/Engine/Documentation/Source/Programming/Development/BuildGraph/ScriptAnatomy/Tasks/PredefinedTasks/BuildGraphScriptPredefinedTasks.INT.udn
@@ -145,13 +145,14 @@ Copies files from one directory to another.
 
 Deploys a tool update through Horde
 
-|         |           |          |                                                                                   |
-| ------- | --------- | -------- | --------------------------------------------------------------------------------- |
-| Name    | String    | Required | Name for the artifact                                                             |
-| Type    | String    | Optional | Type of the artifact                                                              |
-| BaseDir | String    | Required | Base directory to resolve relative paths for input files.                         |
-| Files   | FileSpec  | Optional | Files to be uploaded.                                                             |
-| If      | Condition | Optional | Whether to execute this task. It is ignored if this condition evaluates to false. |
+|             |           |          |                                                                                   |
+| ----------- | --------- | -------- | --------------------------------------------------------------------------------- |
+| Name        | String    | Required | Name for the artifact                                                             |
+| Type        | String    | Optional | Type of the artifact                                                              |
+| Description | String    | Optional | Description for the artifact                                                      |
+| BaseDir     | String    | Required | Base directory to resolve relative paths for input files.                         |
+| Files       | FileSpec  | Optional | Files to be uploaded.                                                             |
+| If          | Condition | Optional | Whether to execute this task. It is ignored if this condition evaluates to false. |
 
 ### CsCompile
 
diff --git a/Engine/Documentation/Source/Programming/UnrealBuildSystem/Configuration/BuildConfigProperties/BuildConfigProperties.INT.udn b/Engine/Documentation/Source/Programming/UnrealBuildSystem/Configuration/BuildConfigProperties/BuildConfigProperties.INT.udn
--- a/Engine/Documentation/Source/Programming/UnrealBuildSystem/Configuration/BuildConfigProperties/BuildConfigProperties.INT.udn
+++ b/Engine/Documentation/Source/Programming/UnrealBuildSystem/Configuration/BuildConfigProperties/BuildConfigProperties.INT.udn
@@ -450,6 +450,8 @@ $ bUseKnownInputs : Set to true to provide known inputs to processes that are ru
 
 $ ActionsOutputFile : Write yaml file with all actions that are queued for build. This can be used to replay using "UbaCli.exe local file.yaml"
 
+$ bDetailedLog : Set to true to see more info about what is happening inside uba and also log output from agents
+
 $ HordeMaxCores : Max cores allowed to be used by build session
 
 $ HordeDelay : How long UBT should wait to ask for help. Useful in build configs where machine can delay remote work and still get same wall time results (pch dependencies etc)
