Documentation for build tools

diff --git a/Engine/Documentation/Source/Programming/UnrealBuildSystem/Configuration/BuildConfigProperties/BuildConfigProperties.INT.udn b/Engine/Documentation/Source/Programming/UnrealBuildSystem/Configuration/BuildConfigProperties/BuildConfigProperties.INT.udn
--- a/Engine/Documentation/Source/Programming/UnrealBuildSystem/Configuration/BuildConfigProperties/BuildConfigProperties.INT.udn
+++ b/Engine/Documentation/Source/Programming/UnrealBuildSystem/Configuration/BuildConfigProperties/BuildConfigProperties.INT.udn
@@ -60,10 +60,6 @@ $ DefaultWarningLevel : Default treatment of uncategorized warnings
 
 $ DeprecationWarningLevel : Level to report deprecation warnings as errors
 
-$ bStaticAnalyzerProjectOnly : Only run static analysis against project modules, skipping engine modules
-
-$ bStaticAnalyzerIncludeGenerated : When enabled, generated source files will be analyzed
-
 $ bDebugBuildsActuallyUseDebugCRT : Enables the debug C++ runtime (CRT) for debug builds. By default we always use the release runtime, since the debug version isn't particularly useful when debugging Unreal Engine projects, and linking against the debug CRT libraries forces our third party library dependencies to also be compiled using the debug CRT (and often perform more slowly). Often it can be inconvenient to require a separate copy of the debug versions of third party static libraries simply so that you can debug your program's code.
 
 $ bLegalToDistributeBinary : Whether the output from this target can be publicly distributed, even if it has dependencies on modules that are in folders with special restrictions (eg. CarefullyRedist, NotForLicensees, NoRedist).
@@ -152,6 +148,12 @@ $ StaticAnalyzerOutputType : The output type to use for the static analyzer. Thi
 
 $ StaticAnalyzerMode : The mode to use for the static analyzer. This is only supported for Clang. Shallow mode completes quicker but is generally not recommended.
 
+$ StaticAnalyzerPVSPrintLevel : The level of warnings to print when analyzing using PVS-Studio
+
+$ bStaticAnalyzerProjectOnly : Only run static analysis against project modules, skipping engine modules
+
+$ bStaticAnalyzerIncludeGenerated : When enabled, generated source files will be analyzed
+
 $ MinFilesUsingPrecompiledHeader : The minimum number of files that must use a pre-compiled header before it will be created and used.
 
 $ bForcePrecompiledHeaderForGameModules : When enabled, a precompiled header is always generated for game modules, even if there are only a few source files in the module.  This greatly improves compile times for iterative changes on a few files in the project, at the expense of slower full rebuild times for small game projects.  This can be overridden by setting MinFilesUsingPrecompiledHeaderOverride in a module's Build.cs file.
diff --git a/Engine/Documentation/Source/Programming/UnrealBuildSystem/TargetFiles/TargetFilesProperties/TargetFilesProperties.INT.udn b/Engine/Documentation/Source/Programming/UnrealBuildSystem/TargetFiles/TargetFilesProperties/TargetFilesProperties.INT.udn
--- a/Engine/Documentation/Source/Programming/UnrealBuildSystem/TargetFiles/TargetFilesProperties/TargetFilesProperties.INT.udn
+++ b/Engine/Documentation/Source/Programming/UnrealBuildSystem/TargetFiles/TargetFilesProperties/TargetFilesProperties.INT.udn
@@ -68,10 +68,6 @@ $ InternalPluginDependencies (List&lt;String&gt;): A list of Plugin names that a
 
 $ DisableMergingModuleAndGeneratedFilesInUnityFiles (String[]): List of modules that disables merging module and generated cpp files in the same unity files.
 
-$ bStaticAnalyzerProjectOnly (Boolean): Only run static analysis against project modules, skipping engine modules
-
-$ bStaticAnalyzerIncludeGenerated (Boolean): When enabled, generated source files will be analyzed
-
 $ DisableUnityBuildForModules (String[]): List of modules to disable unity builds for
 
 $ EnableOptimizeCodeForModules (String[]): List of modules to enable optimizations for
@@ -424,6 +420,12 @@ $ StaticAnalyzerOutputType (StaticAnalyzerOutputType): The output type to use fo
 
 $ StaticAnalyzerMode (StaticAnalyzerMode): The mode to use for the static analyzer. This is only supported for Clang. Shallow mode completes quicker but is generally not recommended.
 
+$ StaticAnalyzerPVSPrintLevel (Int32): The level of warnings to print when analyzing using PVS-Studio
+
+$ bStaticAnalyzerProjectOnly (Boolean): Only run static analysis against project modules, skipping engine modules
+
+$ bStaticAnalyzerIncludeGenerated (Boolean): When enabled, generated source files will be analyzed
+
 $ MinFilesUsingPrecompiledHeader (Int32): The minimum number of files that must use a pre-compiled header before it will be created and used.
 
 $ bForcePrecompiledHeaderForGameModules (Boolean): When enabled, a precompiled header is always generated for game modules, even if there are only a few source files in the module.  This greatly improves compile times for iterative changes on a few files in the project, at the expense of slower full rebuild times for small game projects.  This can be overridden by setting MinFilesUsingPrecompiledHeaderOverride in a module's Build.cs file.
