Documentation for build tools

diff --git a/Engine/Documentation/Source/Programming/Development/BuildGraph/ScriptAnatomy/Tasks/PredefinedTasks/BuildGraphScriptPredefinedTasks.INT.udn b/Engine/Documentation/Source/Programming/Development/BuildGraph/ScriptAnatomy/Tasks/PredefinedTasks/BuildGraphScriptPredefinedTasks.INT.udn
--- a/Engine/Documentation/Source/Programming/Development/BuildGraph/ScriptAnatomy/Tasks/PredefinedTasks/BuildGraphScriptPredefinedTasks.INT.udn
+++ b/Engine/Documentation/Source/Programming/Development/BuildGraph/ScriptAnatomy/Tasks/PredefinedTasks/BuildGraphScriptPredefinedTasks.INT.udn
@@ -145,14 +145,13 @@ Copies files from one directory to another.
 
 Deploys a tool update through Horde
 
-|             |           |          |                                                                                   |
-| ----------- | --------- | -------- | --------------------------------------------------------------------------------- |
-| Name        | String    | Required | Name for the artifact                                                             |
-| Type        | String    | Optional | Type of the artifact                                                              |
-| Description | String    | Optional | Description for the artifact                                                      |
-| BaseDir     | String    | Required | Base directory to resolve relative paths for input files.                         |
-| Files       | FileSpec  | Optional | Files to be uploaded.                                                             |
-| If          | Condition | Optional | Whether to execute this task. It is ignored if this condition evaluates to false. |
+|         |           |          |                                                                                   |
+| ------- | --------- | -------- | --------------------------------------------------------------------------------- |
+| Name    | String    | Required | Name for the artifact                                                             |
+| Type    | String    | Optional | Type of the artifact                                                              |
+| BaseDir | String    | Required | Base directory to resolve relative paths for input files.                         |
+| Files   | FileSpec  | Optional | Files to be uploaded.                                                             |
+| If      | Condition | Optional | Whether to execute this task. It is ignored if this condition evaluates to false. |
 
 ### CsCompile
 
diff --git a/Engine/Documentation/Source/Programming/UnrealBuildSystem/Configuration/BuildConfigProperties/BuildConfigProperties.INT.udn b/Engine/Documentation/Source/Programming/UnrealBuildSystem/Configuration/BuildConfigProperties/BuildConfigProperties.INT.udn
--- a/Engine/Documentation/Source/Programming/UnrealBuildSystem/Configuration/BuildConfigProperties/BuildConfigProperties.INT.udn
+++ b/Engine/Documentation/Source/Programming/UnrealBuildSystem/Configuration/BuildConfigProperties/BuildConfigProperties.INT.udn
@@ -148,12 +148,6 @@ $ StaticAnalyzerOutputType : The output type to use for the static analyzer. Thi
 
 $ StaticAnalyzerMode : The mode to use for the static analyzer. This is only supported for Clang. Shallow mode completes quicker but is generally not recommended.
 
-$ StaticAnalyzerPVSPrintLevel : The level of warnings to print when analyzing using PVS-Studio
-
-$ bStaticAnalyzerProjectOnly : Only run static analysis against project modules, skipping engine modules
-
-$ bStaticAnalyzerIncludeGenerated : When enabled, generated source files will be analyzed
-
 $ MinFilesUsingPrecompiledHeader : The minimum number of files that must use a pre-compiled header before it will be created and used.
 
 $ bForcePrecompiledHeaderForGameModules : When enabled, a precompiled header is always generated for game modules, even if there are only a few source files in the module.  This greatly improves compile times for iterative changes on a few files in the project, at the expense of slower full rebuild times for small game projects.  This can be overridden by setting MinFilesUsingPrecompiledHeaderOverride in a module's Build.cs file.
@@ -450,10 +444,6 @@ $ bUseCrypto : Set to true to enable encryption when transfering files over the
 
 $ bUseKnownInputs : Set to true to provide known inputs to processes that are run remote. This is an experimental feature to speed up build times when ping is higher
 
-$ ActionsOutputFile : Write yaml file with all actions that are queued for build. This can be used to replay using "UbaCli.exe local file.yaml"
-
-$ bDetailedLog : Set to true to see more info about what is happening inside uba and also log output from agents
-
 $ HordeMaxCores : Max cores allowed to be used by build session
 
 $ HordeDelay : How long UBT should wait to ask for help. Useful in build configs where machine can delay remote work and still get same wall time results (pch dependencies etc)
diff --git a/Engine/Documentation/Source/Programming/UnrealBuildSystem/ModuleFiles/ModuleFilesProperties/ModuleFilesProperties.INT.udn b/Engine/Documentation/Source/Programming/UnrealBuildSystem/ModuleFiles/ModuleFilesProperties/ModuleFilesProperties.INT.udn
--- a/Engine/Documentation/Source/Programming/UnrealBuildSystem/ModuleFiles/ModuleFilesProperties/ModuleFilesProperties.INT.udn
+++ b/Engine/Documentation/Source/Programming/UnrealBuildSystem/ModuleFiles/ModuleFilesProperties/ModuleFilesProperties.INT.udn
@@ -12,8 +12,6 @@ $ Logger (ILogger): Accessor for the target logger
 
 $ HasPackageOverride (Boolean): Returns true if there has been an override type specified on this module
 
-$ SDKVersionRelevantPlatforms (List&lt;UnrealTargetPlatform&gt;): If this contains a platform and the project has had it's SDK version overriden from default, this module will be compiled as a project module, instead of shared engine module
-
 $ AutoSdkDirectory (String): The AutoSDK directory for the active host platform
 
 $ EngineDirectory (String): The current engine directory
diff --git a/Engine/Documentation/Source/Programming/UnrealBuildSystem/TargetFiles/TargetFilesProperties/TargetFilesProperties.INT.udn b/Engine/Documentation/Source/Programming/UnrealBuildSystem/TargetFiles/TargetFilesProperties/TargetFilesProperties.INT.udn
--- a/Engine/Documentation/Source/Programming/UnrealBuildSystem/TargetFiles/TargetFilesProperties/TargetFilesProperties.INT.udn
+++ b/Engine/Documentation/Source/Programming/UnrealBuildSystem/TargetFiles/TargetFilesProperties/TargetFilesProperties.INT.udn
@@ -420,12 +420,6 @@ $ StaticAnalyzerOutputType (StaticAnalyzerOutputType): The output type to use fo
 
 $ StaticAnalyzerMode (StaticAnalyzerMode): The mode to use for the static analyzer. This is only supported for Clang. Shallow mode completes quicker but is generally not recommended.
 
-$ StaticAnalyzerPVSPrintLevel (Int32): The level of warnings to print when analyzing using PVS-Studio
-
-$ bStaticAnalyzerProjectOnly (Boolean): Only run static analysis against project modules, skipping engine modules
-
-$ bStaticAnalyzerIncludeGenerated (Boolean): When enabled, generated source files will be analyzed
-
 $ MinFilesUsingPrecompiledHeader (Int32): The minimum number of files that must use a pre-compiled header before it will be created and used.
 
 $ bForcePrecompiledHeaderForGameModules (Boolean): When enabled, a precompiled header is always generated for game modules, even if there are only a few source files in the module.  This greatly improves compile times for iterative changes on a few files in the project, at the expense of slower full rebuild times for small game projects.  This can be overridden by setting MinFilesUsingPrecompiledHeaderOverride in a module's Build.cs file.
@@ -542,8 +536,6 @@ $ ExportPublicHeader (String): Specifies the path to write a header containing p
 
 $ BuildEnvironment (TargetBuildEnvironment): Specifies the build environment for this target. See TargetBuildEnvironment for more information on the available options.
 
-$ bAllowSDKOverrideModulesWithSharedEnvironment (Boolean): If this true, a Shared build environment target will allow for any modules that are SDK version sensitive to have a project side module when an SDK is overridden. For instance, if IOSTargetPlatform, which is marked as IOS SDK version relevant, is compiled for a Target with this set to true, and that target overrides the IOS SDK, then this would compile to e.g. MyProject/Binaries/Win64/IOS/MyProject-IOSTargetPlatform.dll
-
 $ bOverrideBuildEnvironment (Boolean): Whether to ignore violations to the shared build environment (eg. editor targets modifying definitions)
 
 $ AdditionalCompilerArguments (String): Additional arguments to pass to the compiler
