Documentation for build tools

diff --git a/Engine/Documentation/Source/Programming/Development/BuildGraph/ScriptAnatomy/Tasks/PredefinedTasks/BuildGraphScriptPredefinedTasks.INT.udn b/Engine/Documentation/Source/Programming/Development/BuildGraph/ScriptAnatomy/Tasks/PredefinedTasks/BuildGraphScriptPredefinedTasks.INT.udn
--- a/Engine/Documentation/Source/Programming/Development/BuildGraph/ScriptAnatomy/Tasks/PredefinedTasks/BuildGraphScriptPredefinedTasks.INT.udn
+++ b/Engine/Documentation/Source/Programming/Development/BuildGraph/ScriptAnatomy/Tasks/PredefinedTasks/BuildGraphScriptPredefinedTasks.INT.udn
@@ -151,6 +151,7 @@ Deploys a tool update through Horde
 | Type        | String    | Optional | Type of the artifact                                                              |
 | Description | String    | Optional | Description for the artifact                                                      |
 | BaseDir     | String    | Required | Base directory to resolve relative paths for input files.                         |
+| Change      | Integer   | Optional | Changelist number for this artifact                                               |
 | Files       | FileSpec  | Optional | Files to be uploaded.                                                             |
 | If          | Condition | Optional | Whether to execute this task. It is ignored if this condition evaluates to false. |
 
