Documentation for build tools

diff --git a/Engine/Documentation/Source/Programming/UnrealBuildSystem/Configuration/BuildConfigProperties/BuildConfigProperties.INT.udn b/Engine/Documentation/Source/Programming/UnrealBuildSystem/Configuration/BuildConfigProperties/BuildConfigProperties.INT.udn
--- a/Engine/Documentation/Source/Programming/UnrealBuildSystem/Configuration/BuildConfigProperties/BuildConfigProperties.INT.udn
+++ b/Engine/Documentation/Source/Programming/UnrealBuildSystem/Configuration/BuildConfigProperties/BuildConfigProperties.INT.udn
@@ -12,21 +12,17 @@ $ bPrintDebugInfo : Whether debug info should be written to the console.
 
 $ bAllowHybridExecutor : Whether the hybrid executor will be used (a remote executor and local executor). No longer supported.
 
-$ RemoteExecutorPriority : Priority order for remote executors (XGE, SNDBS, FASTBuild, UBAExecutor)
+$ RemoteExecutorPriority : Priority order for remote executors (XGE, SNDBS, FASTBuild, UBA)
 
-$ bAllowBoxExecutor : This property is being renamed and should not be used
+$ bAllowUBAExecutor : Whether the UnrealBuildAccelerator executor will be used.
 
-$ bAllowUBAExecutor : Whether the experimental UnrealBuildAccelerator executor will be used.
-
-$ bAllowBoxLocalExecutor : This property is being renamed and should not be used
-
-$ bAllowUBALocalExecutor : Whether the experimental UnrealBuildAccelerator (local only) executor will be used.
+$ bAllowUBALocalExecutor : Whether the UnrealBuildAccelerator (local only) executor will be used.
 
 $ bAllowXGE : Whether XGE may be used if available, default is true.
 
-$ bAllowFASTBuild : Whether FASTBuild may be used if availabe, default is true.
+$ bAllowFASTBuild : Whether FASTBuild may be used if available, default is true.
 
-$ bAllowSNDBS : Whether SN-DBS may be used if availabe, default is true.
+$ bAllowSNDBS : Whether SN-DBS may be used if available, default is true.
 
 $ bUseUBTMakefiles : Enables support for very fast iterative builds by caching target data. Turning this on causes Unreal Build Tool to emit 'UBT Makefiles' for targets when they are built the first time. Subsequent builds will load these Makefiles and begin outdatedness checking and build invocation very quickly. The caveat is that if source files are added or removed to the project, UBT will need to gather information about those in order for your build to complete successfully. Currently, you must run the project file generator after adding/removing source files to tell UBT to re-gather this information. 
     * Events that can invalidate the 'UBT Makefile':
@@ -292,7 +288,7 @@ $ bWarnOnCppUnrealHeaderTool : If true, generate warnings when C++ UHT is used
 
 $ MaxRootPathLength : Maximum recommended root path length.
 
-$ MaxNestedPathLength : Maximum length of a path relative to the root directory. Used on Windows to ensure paths are portable between machines. Defaults to off.
+$ MaxNestedPathLength : Maximum length of a path relative to the root directory. Used on Windows to ensure paths are portable between machines.
 
 $ bIgnoreStalePGOData : If -PGOOptimize is specified but the linker flags have changed since the last -PGOProfile, this will emit a warning and build without PGO instead of failing during link with LNK1268.
 
@@ -476,7 +472,7 @@ $ bDisableCustomAlloc : Set to true to disable mimalloc and detouring of memory
 
 $ Zone : The zone to use for UBA.
 
-$ bUseCrypto : Set to true to enable encryption when transfering files over the network.
+$ bUseCrypto : Set to true to enable encryption when transferring files over the network.
 
 $ bUseKnownInputs : Set to true to provide known inputs to processes that are run remote. This is an experimental feature to speed up build times when ping is higher
 
@@ -496,7 +492,7 @@ $ bDisableHorde : Disable horde all together
 
 $ Server : Uri of the Horde server
 
-$ Token : Uri of the Horde server
+$ Token : Auth token for the Horde server
 
 $ OidcProvider : OIDC id for the login to use
 
