Documentation for build tools

diff --git a/Engine/Documentation/Source/Programming/Development/BuildGraph/ScriptAnatomy/Tasks/PredefinedTasks/BuildGraphScriptPredefinedTasks.INT.udn b/Engine/Documentation/Source/Programming/Development/BuildGraph/ScriptAnatomy/Tasks/PredefinedTasks/BuildGraphScriptPredefinedTasks.INT.udn
--- a/Engine/Documentation/Source/Programming/Development/BuildGraph/ScriptAnatomy/Tasks/PredefinedTasks/BuildGraphScriptPredefinedTasks.INT.udn
+++ b/Engine/Documentation/Source/Programming/Development/BuildGraph/ScriptAnatomy/Tasks/PredefinedTasks/BuildGraphScriptPredefinedTasks.INT.udn
@@ -108,6 +108,7 @@ Compiles a target with UnrealBuildTool.
 | AllowParallelExecutor | Boolean                   | Optional | No longer necessary as UnrealBuildTool is run to compile targets.                                                                |
 | AllowAllCores         | Boolean                   | Optional | Whether to allow UBT to use all available cores, when AllowXGE is disabled.                                                      |
 | Clean                 | Boolean                   | Optional | Whether to allow cleaning this target. If unspecified, targets are cleaned if the -Clean argument is passed on the command line. |
+| SkipBuild             | Boolean                   | Optional | Global flag passed to UBT that can be used to generate target files without fully compiling.                                     |
 | Tag                   | TagList                   | Optional | Tag to be applied to build products of this task.                                                                                |
 | If                    | Condition                 | Optional | Whether to execute this task. It is ignored if this condition evaluates to false.                                                |
 
@@ -732,6 +733,7 @@ Exports an snapshot from Zen to a specified destination.
 | DestinationFileDir                 | DirectoryReference | Optional | The directory to use when exporting to a file destination                                                                    |
 | DestinationFileName                | String             | Optional | The filename to use when exporting to a file destination                                                                     |
 | OverridePlatformCookedDir          | String             | Optional | Optional. Where to look for the ue.projectstore The pattern {Platform} can be used for exporting multiple platforms at once. |
+| Force                              | Boolean            | Optional | Optional. Whether to force export of data even if the destination claims to have them.                                       |
 | If                                 | Condition          | Optional | Whether to execute this task. It is ignored if this condition evaluates to false.                                            |
 
 ### ZenImportOplog
diff --git a/Engine/Documentation/Source/Programming/UnrealBuildSystem/Configuration/BuildConfigProperties/BuildConfigProperties.INT.udn b/Engine/Documentation/Source/Programming/UnrealBuildSystem/Configuration/BuildConfigProperties/BuildConfigProperties.INT.udn
--- a/Engine/Documentation/Source/Programming/UnrealBuildSystem/Configuration/BuildConfigProperties/BuildConfigProperties.INT.udn
+++ b/Engine/Documentation/Source/Programming/UnrealBuildSystem/Configuration/BuildConfigProperties/BuildConfigProperties.INT.udn
@@ -294,6 +294,8 @@ $ bIgnoreStalePGOData : If -PGOOptimize is specified but the linker flags have c
 
 $ bUseFastGenProfile : If specified along with -PGOProfile, then /FASTGENPROFILE will be used instead of /GENPROFILE. This usually means that the PGO data is generated faster, but the resulting data may not yield as efficient optimizations during -PGOOptimize
 
+$ PreMergedPgdFilename : If specified along with -PGOOptimize, will use the specified per-merged pgd file instead of the usual pgd file with loose pgc files.
+
 $ bPGONoExtraCounters : If specified along with -PGOProfile, prevent the usage of extra counters. Please note that by default /FASTGENPROFILE doesnt use extra counters
 
 $ bSampleBasedPGO : If specified along with -PGOProfile, use sample-based PGO instead of instrumented. Currently Intel oneAPI 2024.0+ only.
