Documentation for build tools

diff --git a/Engine/Documentation/Source/Programming/UnrealBuildSystem/Configuration/BuildConfigProperties/BuildConfigProperties.INT.udn b/Engine/Documentation/Source/Programming/UnrealBuildSystem/Configuration/BuildConfigProperties/BuildConfigProperties.INT.udn
--- a/Engine/Documentation/Source/Programming/UnrealBuildSystem/Configuration/BuildConfigProperties/BuildConfigProperties.INT.udn
+++ b/Engine/Documentation/Source/Programming/UnrealBuildSystem/Configuration/BuildConfigProperties/BuildConfigProperties.INT.udn
@@ -984,3 +984,7 @@ $ bBackupLogFiles : Whether to backup an existing log file, rather than overwrit
 
 $ LogFileBackupCount : The number of log file backups to preserve. Older backups will be deleted.
 
+### Telemetry
+
+$ Providers : Providers to load opt-in telemetry connection information from ini. If unset, or the provider categories do not contain connection info, no telemetry will be sent.
+
diff --git a/Engine/Documentation/Source/Programming/UnrealBuildSystem/TargetFiles/TargetFilesProperties/TargetFilesProperties.INT.udn b/Engine/Documentation/Source/Programming/UnrealBuildSystem/TargetFiles/TargetFilesProperties/TargetFilesProperties.INT.udn
--- a/Engine/Documentation/Source/Programming/UnrealBuildSystem/TargetFiles/TargetFilesProperties/TargetFilesProperties.INT.udn
+++ b/Engine/Documentation/Source/Programming/UnrealBuildSystem/TargetFiles/TargetFilesProperties/TargetFilesProperties.INT.udn
@@ -322,6 +322,8 @@ $ bIsBuildingConsoleApplication (Boolean): True if this is a console application
 
 $ bBuildAdditionalConsoleApp (Boolean): If true, creates an additional console application. Hack for Windows, where it's not possible to conditionally inherit a parent's console Window depending on how the application is invoked; you have to link the same executable with a different subsystem setting.
 
+$ bBuildConsoleAppOnly (Boolean): If true, only creates an additional console application. Overrides bBuildAdditionalConsoleApp.
+
 $ bDisableSymbolCache (Boolean): True if debug symbols that are cached for some platforms should not be created.
 
 $ bUseUnityBuild (Boolean): Whether to unify C++ code into larger files for faster compilation.
