Documentation for build tools

diff --git a/Engine/Documentation/Source/Programming/Development/BuildGraph/ScriptAnatomy/Tasks/PredefinedTasks/BuildGraphScriptPredefinedTasks.INT.udn b/Engine/Documentation/Source/Programming/Development/BuildGraph/ScriptAnatomy/Tasks/PredefinedTasks/BuildGraphScriptPredefinedTasks.INT.udn
--- a/Engine/Documentation/Source/Programming/Development/BuildGraph/ScriptAnatomy/Tasks/PredefinedTasks/BuildGraphScriptPredefinedTasks.INT.udn
+++ b/Engine/Documentation/Source/Programming/Development/BuildGraph/ScriptAnatomy/Tasks/PredefinedTasks/BuildGraphScriptPredefinedTasks.INT.udn
@@ -730,6 +730,7 @@ Exports an snapshot from Zen to a specified destination.
 | SnapshotDescriptorCloudHttpVersion | String             | Optional | The http version to use when writing a snapshot descriptor for a cloud destination                                           |
 | DestinationCloudNamespace          | String             | Optional | The namespace to use when exporting to a cloud destination                                                                   |
 | DestinationCloudBucket             | String             | Optional | A custom bucket name to use when exporting to a cloud destination                                                            |
+| DestinationZenHost                 | String             | Optional | The host name to use when exporting to a zen destination                                                                     |
 | DestinationFileDir                 | DirectoryReference | Optional | The directory to use when exporting to a file destination                                                                    |
 | DestinationFileName                | String             | Optional | The filename to use when exporting to a file destination                                                                     |
 | OverridePlatformCookedDir          | String             | Optional | Optional. Where to look for the ue.projectstore The pattern {Platform} can be used for exporting multiple platforms at once. |
diff --git a/Engine/Documentation/Source/Programming/UnrealBuildSystem/Configuration/BuildConfigProperties/BuildConfigProperties.INT.udn b/Engine/Documentation/Source/Programming/UnrealBuildSystem/Configuration/BuildConfigProperties/BuildConfigProperties.INT.udn
--- a/Engine/Documentation/Source/Programming/UnrealBuildSystem/Configuration/BuildConfigProperties/BuildConfigProperties.INT.udn
+++ b/Engine/Documentation/Source/Programming/UnrealBuildSystem/Configuration/BuildConfigProperties/BuildConfigProperties.INT.udn
@@ -330,6 +330,10 @@ $ PCHMemoryAllocationFactor : Determines the amount of memory that the compiler
 
 $ AdditionalLinkerOptions : Allow the target to specify extra options for linking that aren't otherwise noted here
 
+$ bReducedOptimizeHugeFunctions : Whether to reduce optimizations for huge functions over an instruction threshold to improve compile time https://devblogs.microsoft.com/cppblog/msvc-backend-updates-in-visual-studio-2019-versions-16-3-and-16-4/
+
+$ ReducedOptimizeHugeFunctionsThreshold : The instruction threshold to use when reducing optimizations for huge functions, default 20000.
+
 $ bClangTimeTrace : (Experimental) Appends the -ftime-trace argument to the command line for Clang to output a JSON file containing a timeline for the compile. See http://aras-p.info/blog/2019/01/16/time-trace-timeline-flame-chart-profiler-for-Clang/ for more info.
 
 $ bCompilerTrace : Outputs compile timing information so that it can be analyzed.
