Documentation for build tools

diff --git a/Engine/Documentation/Source/Programming/UnrealBuildSystem/Configuration/BuildConfigProperties/BuildConfigProperties.INT.udn b/Engine/Documentation/Source/Programming/UnrealBuildSystem/Configuration/BuildConfigProperties/BuildConfigProperties.INT.udn
--- a/Engine/Documentation/Source/Programming/UnrealBuildSystem/Configuration/BuildConfigProperties/BuildConfigProperties.INT.udn
+++ b/Engine/Documentation/Source/Programming/UnrealBuildSystem/Configuration/BuildConfigProperties/BuildConfigProperties.INT.udn
@@ -470,6 +470,8 @@ $ bDisableWaitOnMem : Disable UBA waiting on available memory before spawning ne
 
 $ bAllowKillOnMem : Let UBA kill running processes when close to out of memory
 
+$ bStoreObjFilesCompressed : Store object (.obj) compressed on disk. Requires uba to do link step where it will decompress obj files again
+
 $ OutputStatsThresholdMs : Threshold for when executor should output logging for the process. Defaults to never
 
 $ bWriteToDisk : Skip writing intermediate and output files to disk. Useful for validation builds where we don't need the output
