Documentation for build tools

diff --git a/Engine/Documentation/Source/Programming/UnrealBuildSystem/Configuration/BuildConfigProperties/BuildConfigProperties.INT.udn b/Engine/Documentation/Source/Programming/UnrealBuildSystem/Configuration/BuildConfigProperties/BuildConfigProperties.INT.udn
--- a/Engine/Documentation/Source/Programming/UnrealBuildSystem/Configuration/BuildConfigProperties/BuildConfigProperties.INT.udn
+++ b/Engine/Documentation/Source/Programming/UnrealBuildSystem/Configuration/BuildConfigProperties/BuildConfigProperties.INT.udn
@@ -350,6 +350,8 @@ $ ToolchainVersionWarningLevel : Warning level when reporting toolchains that ar
 
 $ bStrictConformanceMode : Enables strict standard conformance mode (/permissive-).
 
+$ bDisableVolatileMetadata : Volatile Metadata is enabled by default and improves x64 emulation on arm64, but may come at a small perfomance cost (/volatileMetadata-).
+
 ### ModuleConfiguration
 
 $ DisableMergingModuleAndGeneratedFilesInUnityFiles : List of modules that disables merging module and generated cpp files in the same unity files.
