Documentation for build tools

diff --git a/Engine/Documentation/Source/Programming/UnrealBuildSystem/Configuration/BuildConfigProperties/BuildConfigProperties.INT.udn b/Engine/Documentation/Source/Programming/UnrealBuildSystem/Configuration/BuildConfigProperties/BuildConfigProperties.INT.udn
--- a/Engine/Documentation/Source/Programming/UnrealBuildSystem/Configuration/BuildConfigProperties/BuildConfigProperties.INT.udn
+++ b/Engine/Documentation/Source/Programming/UnrealBuildSystem/Configuration/BuildConfigProperties/BuildConfigProperties.INT.udn
@@ -220,6 +220,8 @@ $ MinCpuArchX64 : Direct the compiler to generate AVX instructions wherever SSE
 
 $ ActionStallReportTime : Number of second of no completed actions to trigger an action stall report. If zero, stall reports will not be enabled.
 
+$ ActionStallTerminateTime : Number of second of no completed actions to trigger to terminate the queue. If zero, force termination not be enabled.
+
 $ bStopSNDBSCompilationAfterErrors : When enabled, SN-DBS will stop compiling targets after a compile error occurs.  Recommended, as it saves computing resources for others.
 
 $ bXGENoWatchdogThread : Whether to use the no_watchdog_thread option to prevent VS2015 toolchain stalls.
