Documentation for build tools

diff --git a/Engine/Documentation/Source/Programming/UnrealBuildSystem/ModuleFiles/ModuleFilesProperties/ModuleFilesProperties.INT.udn b/Engine/Documentation/Source/Programming/UnrealBuildSystem/ModuleFiles/ModuleFilesProperties/ModuleFilesProperties.INT.udn
--- a/Engine/Documentation/Source/Programming/UnrealBuildSystem/ModuleFiles/ModuleFilesProperties/ModuleFilesProperties.INT.udn
+++ b/Engine/Documentation/Source/Programming/UnrealBuildSystem/ModuleFiles/ModuleFilesProperties/ModuleFilesProperties.INT.udn
@@ -14,6 +14,12 @@ $ HasPackageOverride (Boolean): Returns true if there has been an override type
 
 $ SDKVersionRelevantPlatforms (List&lt;UnrealTargetPlatform&gt;): If this contains a platform and the project has had it's SDK version overriden from default, this module will be compiled as a project module, instead of shared engine module
 
+$ StaticAnalyzerCheckers (HashSet&lt;String&gt;): The static analyzer checkers that should be enabled rather than the defaults. This is only supported for Clang. See https://clang.llvm.org/docs/analyzer/checkers.html for a full list. Or run: 'clang -Xclang -analyzer-checker-help' or: 'clang -Xclang -analyzer-checker-help-alpha' for the list of experimental checkers.
+
+$ StaticAnalyzerDisabledCheckers (HashSet&lt;String&gt;): The static analyzer default checkers that should be disabled. Unused if StaticAnalyzerCheckers is populated. This is only supported for Clang. See https://clang.llvm.org/docs/analyzer/checkers.html for a full list. Or run: 'clang -Xclang -analyzer-checker-help' or: 'clang -Xclang -analyzer-checker-help-alpha' for the list of experimental checkers.
+
+$ StaticAnalyzerAdditionalCheckers (HashSet&lt;String&gt;): The static analyzer non-default checkers that should be enabled. Unused if StaticAnalyzerCheckers is populated. This is only supported for Clang. See https://clang.llvm.org/docs/analyzer/checkers.html for a full list. Or run: 'clang -Xclang -analyzer-checker-help' or: 'clang -Xclang -analyzer-checker-help-alpha' for the list of experimental checkers.
+
 $ AutoSdkDirectory (String): The AutoSDK directory for the active host platform
 
 $ EngineDirectory (String): The current engine directory
@@ -30,12 +36,6 @@ $ IsVcPackageSupported (Boolean): Returns if VcPkg is supported for the build co
 ### Read/Write Properties
 $ StaticAnalyzerRulesets (HashSet&lt;FileReference&gt;): The static analyzer rulesets that should be used to filter warnings. This is only supported for MSVC. See https://learn.microsoft.com/en-us/cpp/code-quality/using-rule-sets-to-specify-the-cpp-rules-to-run
 
-$ StaticAnalyzerCheckers (HashSet&lt;String&gt;): The static analyzer checkers that should be enabled rather than the defaults. This is only supported for Clang. See https://clang.llvm.org/docs/analyzer/checkers.html for a full list. Or run: 'clang -Xclang -analyzer-checker-help' or: 'clang -Xclang -analyzer-checker-help-alpha' for the list of experimental checkers.
-
-$ StaticAnalyzerDisabledCheckers (HashSet&lt;String&gt;): The static analyzer default checkers that should be disabled. Unused if StaticAnalyzerCheckers is populated. This is only supported for Clang. This overrides the default disabled checkers, which are deadcode.DeadStores and security.FloatLoopCounter See https://clang.llvm.org/docs/analyzer/checkers.html for a full list. Or run: 'clang -Xclang -analyzer-checker-help' or: 'clang -Xclang -analyzer-checker-help-alpha' for the list of experimental checkers.
-
-$ StaticAnalyzerAdditionalCheckers (HashSet&lt;String&gt;): The static analyzer non-default checkers that should be enabled. Unused if StaticAnalyzerCheckers is populated. This is only supported for Clang. See https://clang.llvm.org/docs/analyzer/checkers.html for a full list. Or run: 'clang -Xclang -analyzer-checker-help' or: 'clang -Xclang -analyzer-checker-help-alpha' for the list of experimental checkers.
-
 $ AllowedRestrictedFolders (List&lt;String&gt;): List of folders which are allowed to be referenced when compiling this binary, without propagating restricted folder names
 
 $ AliasRestrictedFolders (Dictionary&lt;String, String&gt;): Set of aliased restricted folder references
