Documentation for build tools

diff --git a/Engine/Documentation/Source/Programming/Development/BuildGraph/ScriptAnatomy/Tasks/PredefinedTasks/BuildGraphScriptPredefinedTasks.INT.udn b/Engine/Documentation/Source/Programming/Development/BuildGraph/ScriptAnatomy/Tasks/PredefinedTasks/BuildGraphScriptPredefinedTasks.INT.udn
--- a/Engine/Documentation/Source/Programming/Development/BuildGraph/ScriptAnatomy/Tasks/PredefinedTasks/BuildGraphScriptPredefinedTasks.INT.udn
+++ b/Engine/Documentation/Source/Programming/Development/BuildGraph/ScriptAnatomy/Tasks/PredefinedTasks/BuildGraphScriptPredefinedTasks.INT.udn
@@ -193,7 +193,7 @@ Deploys a tool update through Horde
 |           |           |          |                                                                                                         |
 | --------- | --------- | -------- | ------------------------------------------------------------------------------------------------------- |
 | Id        | String    | Required | Identifier for the tool                                                                                 |
-| Settings  | String    | Required | Settings file to use for the deployment. Should be a JSON file containing server name and access token. |
+| Settings  | String    | Optional | Settings file to use for the deployment. Should be a JSON file containing server name and access token. |
 | Version   | String    | Required | Version number for the new tool                                                                         |
 | Duration  | Integer   | Optional | Duration over which to roll out the tool, in minutes.                                                   |
 | Paused    | Boolean   | Optional | Whether to create the deployment as paused                                                              |
diff --git a/Engine/Documentation/Source/Programming/UnrealBuildSystem/Configuration/BuildConfigProperties/BuildConfigProperties.INT.udn b/Engine/Documentation/Source/Programming/UnrealBuildSystem/Configuration/BuildConfigProperties/BuildConfigProperties.INT.udn
--- a/Engine/Documentation/Source/Programming/UnrealBuildSystem/Configuration/BuildConfigProperties/BuildConfigProperties.INT.udn
+++ b/Engine/Documentation/Source/Programming/UnrealBuildSystem/Configuration/BuildConfigProperties/BuildConfigProperties.INT.udn
@@ -372,66 +372,10 @@
 
 $ bCompileChaosVisualDebuggerSupport : Whether to compile in Chaos Visual Debugger (CVD) support features to record the state of the physics simulation
 
-### TestTargetRules
-
-$ bCompileChaosVisualDebuggerSupport : Whether to compile in Chaos Visual Debugger (CVD) support features to record the state of the physics simulation
-
-### FASTBuild
-
-$ FBuildExecutablePath : Used to specify the location of fbuild.exe if the distributed binary isn't being used
-
-$ bEnableDistribution : Controls network build distribution
-
-$ FBuildBrokeragePath : Used to specify the location of the brokerage. If null, FASTBuild will fall back to checking FASTBUILD_BROKERAGE_PATH
-
-$ FBuildCoordinator : Used to specify the FASTBuild coordinator IP or network name. If null, FASTBuild will fall back to checking FASTBUILD_COORDINATOR
-
-$ bEnableCaching : Controls whether to use caching at all. CachePath and FASTCacheMode are only relevant if this is enabled.
-
-$ CacheMode : Cache access mode - only relevant if bEnableCaching is true;
-
-$ FBuildCachePath : Used to specify the location of the cache. If null, FASTBuild will fall back to checking FASTBUILD_CACHE_PATH
-
-$ bForceRemote : Whether to force remote
-
-$ bStopOnError : Whether to stop on error
-
-$ MsvcCRTRedistVersion : Which MSVC CRT Redist version to use
-
-$ CompilerVersion : Which MSVC Compiler version to use
-
-### ParallelExecutor
-
-$ MaxProcessorCount : Maximum processor count for local execution.
-
-$ ProcessorCountMultiplier : Processor count multiplier for local execution. Can be below 1 to reserve CPU for other tasks. When using the local executor (not XGE), run a single action on each CPU core. Note that you can set this to a larger value to get slightly faster build times in many cases, but your computer's responsiveness during compiling may be much worse. This value is ignored if the CPU does not support hyper-threading.
-
-$ MemoryPerActionBytes : Free memory per action in bytes, used to limit the number of parallel actions if the machine is memory starved. Set to 0 to disable free memory checking.
-
-$ ProcessPriority : The priority to set for spawned processes. Valid Settings: Idle, BelowNormal, Normal, AboveNormal, High Default: BelowNormal or Normal for an Asymmetrical processor as BelowNormal can cause scheduling issues.
-
-$ bStopCompilationAfterErrors : When enabled, will stop compiling targets after a compile error occurs.
-
-$ bShowCompilationTimes : Whether to show compilation times along with worst offenders or not.
-
-$ bShowPerActionCompilationTimes : Whether to show compilation times for each executed action
-
-$ bLogActionCommandLines : Whether to log command lines for actions being executed
-
-$ bPrintActionTargetNames : Add target names for each action executed
-
-$ bUseActionWeights : Whether to take into account the Action's weight when determining to do more work or not.
-
-$ bShowCPUUtilization : Whether to show CPU utilization after the work is complete.
-
-### SNDBS
-
-$ bAllowOverVpn : When set to false, SNDBS will not be enabled when running connected to the coordinator over VPN. Configure VPN-assigned subnets via the VpnSubnets parameter.
-
-$ VpnSubnets : List of subnets containing IP addresses assigned by VPN
-
 ### UnrealBuildAccelerator
 
+$ bAllowUbaCompression : Experimental: Store object (.obj) compressed on disk. Requires UBA to link, currently MSVC only. Toggling this flag will invalidate MSVC actions. Warning, this option is not currently compatitable with PGO or the the cl-clang linker as those are not detoured and linking will fail.
+
 $ bDisableRemote : When set to true, UBA will not use any remote help
 
 $ bForceBuildAllRemote : When set to true, UBA will force all actions that can be built remotely to be built remotely. This will hang if there are no remote agents available
@@ -500,6 +444,64 @@
 
 $ bDisableHorde : Disable horde all together
 
+### TestTargetRules
+
+$ bCompileChaosVisualDebuggerSupport : Whether to compile in Chaos Visual Debugger (CVD) support features to record the state of the physics simulation
+
+### FASTBuild
+
+$ FBuildExecutablePath : Used to specify the location of fbuild.exe if the distributed binary isn't being used
+
+$ bEnableDistribution : Controls network build distribution
+
+$ FBuildBrokeragePath : Used to specify the location of the brokerage. If null, FASTBuild will fall back to checking FASTBUILD_BROKERAGE_PATH
+
+$ FBuildCoordinator : Used to specify the FASTBuild coordinator IP or network name. If null, FASTBuild will fall back to checking FASTBUILD_COORDINATOR
+
+$ bEnableCaching : Controls whether to use caching at all. CachePath and FASTCacheMode are only relevant if this is enabled.
+
+$ CacheMode : Cache access mode - only relevant if bEnableCaching is true;
+
+$ FBuildCachePath : Used to specify the location of the cache. If null, FASTBuild will fall back to checking FASTBUILD_CACHE_PATH
+
+$ bForceRemote : Whether to force remote
+
+$ bStopOnError : Whether to stop on error
+
+$ MsvcCRTRedistVersion : Which MSVC CRT Redist version to use
+
+$ CompilerVersion : Which MSVC Compiler version to use
+
+### ParallelExecutor
+
+$ MaxProcessorCount : Maximum processor count for local execution.
+
+$ ProcessorCountMultiplier : Processor count multiplier for local execution. Can be below 1 to reserve CPU for other tasks. When using the local executor (not XGE), run a single action on each CPU core. Note that you can set this to a larger value to get slightly faster build times in many cases, but your computer's responsiveness during compiling may be much worse. This value is ignored if the CPU does not support hyper-threading.
+
+$ MemoryPerActionBytes : Free memory per action in bytes, used to limit the number of parallel actions if the machine is memory starved. Set to 0 to disable free memory checking.
+
+$ ProcessPriority : The priority to set for spawned processes. Valid Settings: Idle, BelowNormal, Normal, AboveNormal, High Default: BelowNormal or Normal for an Asymmetrical processor as BelowNormal can cause scheduling issues.
+
+$ bStopCompilationAfterErrors : When enabled, will stop compiling targets after a compile error occurs.
+
+$ bShowCompilationTimes : Whether to show compilation times along with worst offenders or not.
+
+$ bShowPerActionCompilationTimes : Whether to show compilation times for each executed action
+
+$ bLogActionCommandLines : Whether to log command lines for actions being executed
+
+$ bPrintActionTargetNames : Add target names for each action executed
+
+$ bUseActionWeights : Whether to take into account the Action's weight when determining to do more work or not.
+
+$ bShowCPUUtilization : Whether to show CPU utilization after the work is complete.
+
+### SNDBS
+
+$ bAllowOverVpn : When set to false, SNDBS will not be enabled when running connected to the coordinator over VPN. Configure VPN-assigned subnets via the VpnSubnets parameter.
+
+$ VpnSubnets : List of subnets containing IP addresses assigned by VPN
+
 ### Horde
 
 $ Server : Uri of the Horde server
diff --git a/Engine/Documentation/Source/Programming/UnrealBuildSystem/TargetFiles/TargetFilesProperties/TargetFilesProperties.INT.udn b/Engine/Documentation/Source/Programming/UnrealBuildSystem/TargetFiles/TargetFilesProperties/TargetFilesProperties.INT.udn
--- a/Engine/Documentation/Source/Programming/UnrealBuildSystem/TargetFiles/TargetFilesProperties/TargetFilesProperties.INT.udn
+++ b/Engine/Documentation/Source/Programming/UnrealBuildSystem/TargetFiles/TargetFilesProperties/TargetFilesProperties.INT.udn
@@ -416,6 +416,8 @@ $ bPreprocessDepends (Boolean): Generate dependency files by preprocessing. This
 
 $ bWithAssembly (Boolean): Whether to generate assembly data while compiling this target. Works exclusively on MSVC for now.
 
+$ bAllowUbaCompression (Boolean): Experimental: Store object (.obj) compressed on disk. Requires UBA to link, currently MSVC only. Toggling this flag will invalidate MSVC actions. Warning, this option is not currently compatitable with PGO or the the cl-clang linker as those are not detoured and linking will fail.
+
 $ StaticAnalyzer (StaticAnalyzer): Whether static code analysis should be enabled.
 
 $ StaticAnalyzerOutputType (StaticAnalyzerOutputType): The output type to use for the static analyzer. This is only supported for Clang.
