Documentation for build tools

diff --git a/Engine/Documentation/Source/Programming/UnrealBuildSystem/Configuration/BuildConfigProperties/BuildConfigProperties.INT.udn b/Engine/Documentation/Source/Programming/UnrealBuildSystem/Configuration/BuildConfigProperties/BuildConfigProperties.INT.udn
--- a/Engine/Documentation/Source/Programming/UnrealBuildSystem/Configuration/BuildConfigProperties/BuildConfigProperties.INT.udn
+++ b/Engine/Documentation/Source/Programming/UnrealBuildSystem/Configuration/BuildConfigProperties/BuildConfigProperties.INT.udn
@@ -436,6 +436,12 @@ $ ActionsOutputFile : Write yaml file with all actions that are queued for build
 
 $ bDetailedLog : Set to true to see more info about what is happening inside uba and also log output from agents
 
+$ Cache : Address of the uba cache service. Will automatically use cache if connected
+
+$ WriteCache : Set cache to write instead of fetch
+
+$ CacheMaxWorkers : Max number of cache download tasks that can execute in parallel
+
 $ HordeMaxCores : Max cores allowed to be used by build session
 
 $ HordeDelay : How long UBT should wait to ask for help. Useful in build configs where machine can delay remote work and still get same wall time results (pch dependencies etc)
