Documentation for build tools

diff --git a/Engine/Documentation/Source/Programming/UnrealBuildSystem/TargetFiles/TargetFilesProperties/TargetFilesProperties.INT.udn b/Engine/Documentation/Source/Programming/UnrealBuildSystem/TargetFiles/TargetFilesProperties/TargetFilesProperties.INT.udn
--- a/Engine/Documentation/Source/Programming/UnrealBuildSystem/TargetFiles/TargetFilesProperties/TargetFilesProperties.INT.udn
+++ b/Engine/Documentation/Source/Programming/UnrealBuildSystem/TargetFiles/TargetFilesProperties/TargetFilesProperties.INT.udn
@@ -240,6 +240,8 @@ $ bEnablePIE (Boolean): Enable Position Independent Executable (PIE). Has an ove
 
 $ bEnableStackProtection (Boolean): Enable Stack Protection. Has an overhead cost
 
+$ bWithClientCode (Boolean): Compile client-only code.
+
 $ bWithServerCode (Boolean): Compile server-only code.
 
 $ bFNameOutlineNumber (Boolean): Compile with FName storing the number part in the name table. Saves memory when most names are not numbered and those that are are referenced multiple times. The game and engine must ensure they reuse numbered names similarly to name strings to avoid leaking memory.
