Documentation for build tools

diff --git a/Engine/Documentation/Source/Programming/UnrealBuildSystem/TargetFiles/TargetFilesProperties/TargetFilesProperties.INT.udn b/Engine/Documentation/Source/Programming/UnrealBuildSystem/TargetFiles/TargetFilesProperties/TargetFilesProperties.INT.udn
--- a/Engine/Documentation/Source/Programming/UnrealBuildSystem/TargetFiles/TargetFilesProperties/TargetFilesProperties.INT.udn
+++ b/Engine/Documentation/Source/Programming/UnrealBuildSystem/TargetFiles/TargetFilesProperties/TargetFilesProperties.INT.udn
@@ -236,6 +236,8 @@ $ bForceEnableObjCExceptions (Boolean): Enable exceptions for all modules.
 
 $ bForceEnableRTTI (Boolean): Enable RTTI for all modules.
 
+$ bEnablePrivateBuildInformation (Boolean): Enable BuildSettings to contain private information about build. Such as machine name, user name and user domain name (exposed in BuildSettings.h)
+
 $ bEnablePIE (Boolean): Enable Position Independent Executable (PIE). Has an overhead cost
 
 $ bEnableStackProtection (Boolean): Enable Stack Protection. Has an overhead cost
