Documentation for build tools

diff --git a/Engine/Documentation/Source/Programming/UnrealBuildSystem/Configuration/BuildConfigProperties/BuildConfigProperties.INT.udn b/Engine/Documentation/Source/Programming/UnrealBuildSystem/Configuration/BuildConfigProperties/BuildConfigProperties.INT.udn
--- a/Engine/Documentation/Source/Programming/UnrealBuildSystem/Configuration/BuildConfigProperties/BuildConfigProperties.INT.udn
+++ b/Engine/Documentation/Source/Programming/UnrealBuildSystem/Configuration/BuildConfigProperties/BuildConfigProperties.INT.udn
@@ -382,6 +382,8 @@ $ bForceBuildAllRemote : When set to true, UBA will force all actions that can b
 
 $ bForcedRetry : When set to true, actions that fail locally with UBA will be retried without UBA.
 
+$ bForcedRetryRemote : When set to true, actions that fail remotely with UBA will be retried locally with UBA.
+
 $ bStrict : When set to true, all errors and warnings from UBA will be output at the appropriate severity level to the log (rather than being output as 'information' and attempting to continue regardless).
 
 $ bStoreRaw : If UBA should store cas compressed or raw
