Documentation for build tools

diff --git a/Engine/Documentation/Source/Programming/UnrealBuildSystem/Configuration/BuildConfigProperties/BuildConfigProperties.INT.udn b/Engine/Documentation/Source/Programming/UnrealBuildSystem/Configuration/BuildConfigProperties/BuildConfigProperties.INT.udn
--- a/Engine/Documentation/Source/Programming/UnrealBuildSystem/Configuration/BuildConfigProperties/BuildConfigProperties.INT.udn
+++ b/Engine/Documentation/Source/Programming/UnrealBuildSystem/Configuration/BuildConfigProperties/BuildConfigProperties.INT.udn
@@ -444,6 +444,8 @@ $ WriteCache : Set cache to write instead of fetch
 
 $ CacheMaxWorkers : Max number of cache download tasks that can execute in parallel
 
+$ ReportCacheMissReason : Report reason a cache miss happened. Useful when searching for determinism/portability issues
+
 $ HordeMaxCores : Max cores allowed to be used by build session
 
 $ HordeDelay : How long UBT should wait to ask for help. Useful in build configs where machine can delay remote work and still get same wall time results (pch dependencies etc)
