Documentation for build tools

diff --git a/Engine/Documentation/Source/Programming/UnrealBuildSystem/TargetFiles/TargetFilesProperties/TargetFilesProperties.INT.udn b/Engine/Documentation/Source/Programming/UnrealBuildSystem/TargetFiles/TargetFilesProperties/TargetFilesProperties.INT.udn
--- a/Engine/Documentation/Source/Programming/UnrealBuildSystem/TargetFiles/TargetFilesProperties/TargetFilesProperties.INT.udn
+++ b/Engine/Documentation/Source/Programming/UnrealBuildSystem/TargetFiles/TargetFilesProperties/TargetFilesProperties.INT.udn
@@ -68,6 +68,8 @@ $ InternalPluginDependencies (List&lt;String&gt;): A list of Plugin names that a
 
 $ DisableMergingModuleAndGeneratedFilesInUnityFiles (String[]): List of modules that disables merging module and generated cpp files in the same unity files.
 
+$ bShowIncludes (Boolean): Print out files that are included by each source file
+
 $ DisableUnityBuildForModules (String[]): List of modules to disable unity builds for
 
 $ EnableOptimizeCodeForModules (String[]): List of modules to enable optimizations for
