Documentation for build tools

diff --git a/Engine/Documentation/Source/Programming/UnrealBuildSystem/TargetFiles/TargetFilesProperties/TargetFilesProperties.INT.udn b/Engine/Documentation/Source/Programming/UnrealBuildSystem/TargetFiles/TargetFilesProperties/TargetFilesProperties.INT.udn
--- a/Engine/Documentation/Source/Programming/UnrealBuildSystem/TargetFiles/TargetFilesProperties/TargetFilesProperties.INT.udn
+++ b/Engine/Documentation/Source/Programming/UnrealBuildSystem/TargetFiles/TargetFilesProperties/TargetFilesProperties.INT.udn
@@ -26,6 +26,10 @@ $ DebugInfoLineTablesOnlyModules (HashSet&lt;String&gt;): Modules that should em
 
 $ DebugInfoLineTablesOnlyPlugins (HashSet&lt;String&gt;): Plugins that should emit line number tables instead of full debug information for compilers that support doing so. Overrides DisableDebugInfoPlugins
 
+$ MergePlugins (List&lt;String&gt;): Experimental: List of plugins (and their dependencies) to each merge into separate libraries. Requires bMergeModules to be enabled
+
+$ MergePluginsShared (Dictionary&lt;String, HashSet&lt;String&gt;&gt;): Experimental: List of merged plugins to move common dependenceis to a shared library, requires MergePlugins. Can be chained
+
 $ GlobalDefinitions (List&lt;String&gt;): Macros to define globally across the whole target.
 
 $ ProjectDefinitions (List&lt;String&gt;): Macros to define across all macros in the project.
@@ -546,6 +550,10 @@ $ BuildVersion (String): The build version string
 
 $ LinkType (TargetLinkType): Specifies how to link modules in this target (monolithic or modular). This is currently protected for backwards compatibility. Call the GetLinkType() accessor until support for the deprecated ShouldCompileMonolithic() override has been removed.
 
+$ bStripExports (Boolean): Experimental: Strip unused exports from libraries. Only applies when LinkType is Modular
+
+$ bMergeModules (Boolean): Experimental: Merge modular modules into combined libraries. Sets LinkType to Modular and enables bStripExports
+
 $ LaunchModuleName (String): Specifies the name of the launch module. For modular builds, this is the module that is compiled into the target's executable.
 
 $ ExportPublicHeader (String): Specifies the path to write a header containing public definitions for this target. Useful when building a DLL to be consumed by external build processes.
