Documentation for build tools

diff --git a/Engine/Documentation/Source/Programming/Development/BuildGraph/ScriptAnatomy/Tasks/PredefinedTasks/BuildGraphScriptPredefinedTasks.INT.udn b/Engine/Documentation/Source/Programming/Development/BuildGraph/ScriptAnatomy/Tasks/PredefinedTasks/BuildGraphScriptPredefinedTasks.INT.udn
--- a/Engine/Documentation/Source/Programming/Development/BuildGraph/ScriptAnatomy/Tasks/PredefinedTasks/BuildGraphScriptPredefinedTasks.INT.udn
+++ b/Engine/Documentation/Source/Programming/Development/BuildGraph/ScriptAnatomy/Tasks/PredefinedTasks/BuildGraphScriptPredefinedTasks.INT.udn
@@ -350,6 +350,16 @@ Creates a Horde report file, which will be displayed on the dashboard with any j
 | Text      | String    | Required | Text to be displayed                                                              |
 | If        | Condition | Optional | Whether to execute this task. It is ignored if this condition evaluates to false. |
 
+### Horde-GetSecrets
+
+Replaces strings in a text file with secrets obtained from Horde
+
+|         |           |          |                                                                                                                               |
+| ------- | --------- | -------- | ----------------------------------------------------------------------------------------------------------------------------- |
+| File    | String    | Required | File to update with secrets                                                                                                   |
+| Replace | String    | Required | Pairs of strings and secret names to expand in the text file, in the form SOURCE_TEXT=secret-name;SOURCE_TEXT_2=secret-name-2 |
+| If      | Condition | Optional | Whether to execute this task. It is ignored if this condition evaluates to false.                                             |
+
 ### Kubectl
 
 Spawns Kubectl and waits for it to complete.
diff --git a/Engine/Documentation/Source/Programming/UnrealBuildSystem/TargetFiles/TargetFilesProperties/TargetFilesProperties.INT.udn b/Engine/Documentation/Source/Programming/UnrealBuildSystem/TargetFiles/TargetFilesProperties/TargetFilesProperties.INT.udn
--- a/Engine/Documentation/Source/Programming/UnrealBuildSystem/TargetFiles/TargetFilesProperties/TargetFilesProperties.INT.udn
+++ b/Engine/Documentation/Source/Programming/UnrealBuildSystem/TargetFiles/TargetFilesProperties/TargetFilesProperties.INT.udn
@@ -138,6 +138,8 @@ $ VSTestRunSettingsFile (FileReference): Set this to reference a VSTest run sett
 
 $ bRuntimeDependenciesComeFromBuildPlugins (Boolean): If this is true, then the BuildPlugins list will be used to populate RuntimeDependencies, rather than EnablePlugins
 
+$ bAllowEnginePluginsEnabledByDefault (Boolean): If false, suppress loading "enabled by default" engine plugins not explicitly enabled by the project or target files.
+
 $ DisablePluginsConflictWarningLevel (WarningLevel): How to treat conflicts when a disabled plugin is being enabled by another plugin referencing it
 
 $ PakSigningKeysFile (String): Path to the set of pak signing keys to embed in the executable.
