Documentation for build tools

diff --git a/Engine/Documentation/Source/Programming/Development/BuildGraph/ScriptAnatomy/Tasks/PredefinedTasks/BuildGraphScriptPredefinedTasks.INT.udn b/Engine/Documentation/Source/Programming/Development/BuildGraph/ScriptAnatomy/Tasks/PredefinedTasks/BuildGraphScriptPredefinedTasks.INT.udn
--- a/Engine/Documentation/Source/Programming/Development/BuildGraph/ScriptAnatomy/Tasks/PredefinedTasks/BuildGraphScriptPredefinedTasks.INT.udn
+++ b/Engine/Documentation/Source/Programming/Development/BuildGraph/ScriptAnatomy/Tasks/PredefinedTasks/BuildGraphScriptPredefinedTasks.INT.udn
@@ -354,11 +354,12 @@ Creates a Horde report file, which will be displayed on the dashboard with any j
 
 Replaces strings in a text file with secrets obtained from Horde
 
-|         |           |          |                                                                                                                               |
-| ------- | --------- | -------- | ----------------------------------------------------------------------------------------------------------------------------- |
-| File    | String    | Required | File to update with secrets                                                                                                   |
-| Replace | String    | Required | Pairs of strings and secret names to expand in the text file, in the form SOURCE_TEXT=secret-name;SOURCE_TEXT_2=secret-name-2 |
-| If      | Condition | Optional | Whether to execute this task. It is ignored if this condition evaluates to false.                                             |
+|         |           |          |                                                                                                                                                                                                                                           |
+| ------- | --------- | -------- | ----------------------------------------------------------------------------------------------------------------------------------------------------------------------------------------------------------------------------------------- |
+| File    | String    | Required | File to update with secrets                                                                                                                                                                                                               |
+| Text    | String    | Optional | Text to update with secrets                                                                                                                                                                                                               |
+| Replace | String    | Optional | Pairs of strings and secret names to expand in the text file, in the form SOURCE_TEXT=secret-name;SOURCE_TEXT_2=secret-name-2. If not specified, secrets embedded inline in the text will be expanded from {{secret-name.value}} strings. |
+| If      | Condition | Optional | Whether to execute this task. It is ignored if this condition evaluates to false.                                                                                                                                                         |
 
 ### Kubectl
 
diff --git a/Engine/Documentation/Source/Programming/UnrealBuildSystem/TargetFiles/TargetFilesProperties/TargetFilesProperties.INT.udn b/Engine/Documentation/Source/Programming/UnrealBuildSystem/TargetFiles/TargetFilesProperties/TargetFilesProperties.INT.udn
--- a/Engine/Documentation/Source/Programming/UnrealBuildSystem/TargetFiles/TargetFilesProperties/TargetFilesProperties.INT.udn
+++ b/Engine/Documentation/Source/Programming/UnrealBuildSystem/TargetFiles/TargetFilesProperties/TargetFilesProperties.INT.udn
@@ -146,6 +146,8 @@ $ PakSigningKeysFile (String): Path to the set of pak signing keys to embed in t
 
 $ SolutionDirectory (String): Allows a Program Target to specify it's own solution folder path.
 
+$ bGenerateProgramProject (Boolean): Force a Target to be treated as a Program for the purposes of project file generation.
+
 $ bBuildInSolutionByDefault (Nullable&lt;Boolean&gt;): Whether the target should be included in the default solution build configuration Setting this to false will skip building when running in the IDE
 
 $ bShouldCompileAsDLL (Boolean): Whether this target should be compiled as a DLL.  Requires LinkType to be set to TargetLinkType.Monolithic.
