Documentation for build tools

diff --git a/Engine/Documentation/Source/Programming/UnrealBuildSystem/ModuleFiles/ModuleFilesProperties/ModuleFilesProperties.INT.udn b/Engine/Documentation/Source/Programming/UnrealBuildSystem/ModuleFiles/ModuleFilesProperties/ModuleFilesProperties.INT.udn
--- a/Engine/Documentation/Source/Programming/UnrealBuildSystem/ModuleFiles/ModuleFilesProperties/ModuleFilesProperties.INT.udn
+++ b/Engine/Documentation/Source/Programming/UnrealBuildSystem/ModuleFiles/ModuleFilesProperties/ModuleFilesProperties.INT.udn
@@ -20,6 +20,8 @@ $ StaticAnalyzerDisabledCheckers (HashSet&lt;String&gt;): The static analyzer de
 
 $ StaticAnalyzerAdditionalCheckers (HashSet&lt;String&gt;): The static analyzer non-default checkers that should be enabled. Unused if StaticAnalyzerCheckers is populated. This is only supported for Clang. See https://clang.llvm.org/docs/analyzer/checkers.html for a full list. Or run: 'clang -Xclang -analyzer-checker-help' or: 'clang -Xclang -analyzer-checker-help-alpha' for the list of experimental checkers.
 
+$ StaticAnalyzerPVSDisabledErrors (HashSet&lt;String&gt;): The PVS Studio analysis errors that should be disabled.
+
 $ AutoSdkDirectory (String): The AutoSDK directory for the active host platform
 
 $ EngineDirectory (String): The current engine directory
