Documentation for build tools

diff --git a/Engine/Documentation/Source/Programming/Development/BuildGraph/ScriptAnatomy/Tasks/PredefinedTasks/BuildGraphScriptPredefinedTasks.INT.udn b/Engine/Documentation/Source/Programming/Development/BuildGraph/ScriptAnatomy/Tasks/PredefinedTasks/BuildGraphScriptPredefinedTasks.INT.udn
--- a/Engine/Documentation/Source/Programming/Development/BuildGraph/ScriptAnatomy/Tasks/PredefinedTasks/BuildGraphScriptPredefinedTasks.INT.udn
+++ b/Engine/Documentation/Source/Programming/Development/BuildGraph/ScriptAnatomy/Tasks/PredefinedTasks/BuildGraphScriptPredefinedTasks.INT.udn
@@ -737,6 +737,7 @@ Exports an snapshot from Zen to a specified destination.
 | DestinationIdentifier              | String             | Optional | The identifier to use when exporting to a destination                                                                        |
 | DestinationCloudHost               | String             | Optional | The host name to use when exporting to a cloud destination                                                                   |
 | SnapshotDescriptorCloudHost        | String             | Optional | The host name to use when writing a snapshot descriptor for a cloud destination                                              |
+| SnapshotDescriptorPlatform         | String             | Optional | The target platform to use when writing a snapshot descriptor                                                                |
 | DestinationCloudHttpVersion        | String             | Optional | The http version to use when exporting to a cloud destination                                                                |
 | SnapshotDescriptorCloudHttpVersion | String             | Optional | The http version to use when writing a snapshot descriptor for a cloud destination                                           |
 | DestinationCloudNamespace          | String             | Optional | The namespace to use when exporting to a cloud destination                                                                   |
@@ -745,7 +746,8 @@ Exports an snapshot from Zen to a specified destination.
 | DestinationFileDir                 | DirectoryReference | Optional | The directory to use when exporting to a file destination                                                                    |
 | DestinationFileName                | String             | Optional | The filename to use when exporting to a file destination                                                                     |
 | OverridePlatformCookedDir          | String             | Optional | Optional. Where to look for the ue.projectstore The pattern {Platform} can be used for exporting multiple platforms at once. |
-| Force                              | Boolean            | Optional | Optional. Whether to force export of data even if the destination claims to have them.                                       |
+| ForceExport                        | Boolean            | Optional | Optional. Whether to force export of data even if the destination claims to have them.                                       |
+| SkipExport                         | Boolean            | Optional | Optional. Whether to entirely bypass the exporting of data and write a snapshot descriptor as if the data had been exported. |
 | If                                 | Condition          | Optional | Whether to execute this task. It is ignored if this condition evaluates to false.                                            |
 
 ### ZenImportOplog
