Documentation for build tools

diff --git a/Engine/Documentation/Source/Programming/UnrealBuildSystem/Configuration/BuildConfigProperties/BuildConfigProperties.INT.udn b/Engine/Documentation/Source/Programming/UnrealBuildSystem/Configuration/BuildConfigProperties/BuildConfigProperties.INT.udn
--- a/Engine/Documentation/Source/Programming/UnrealBuildSystem/Configuration/BuildConfigProperties/BuildConfigProperties.INT.udn
+++ b/Engine/Documentation/Source/Programming/UnrealBuildSystem/Configuration/BuildConfigProperties/BuildConfigProperties.INT.udn
@@ -374,7 +374,7 @@ $ bCompileChaosVisualDebuggerSupport : Whether to compile in Chaos Visual Debugg
 
 ### UnrealBuildAccelerator
 
-$ bAllowUbaCompression : Experimental: Store object (.obj) compressed on disk. Requires UBA to link, currently MSVC only. Toggling this flag will invalidate MSVC actions. Warning, this option is not currently compatitable with PGO or the the cl-clang linker as those are not detoured and linking will fail.
+$ bStoreObjFilesCompressed : Store object (.obj) compressed on disk. Requires uba to do link step where it will decompress obj files again
 
 $ bDisableRemote : When set to true, UBA will not use any remote help
 
@@ -420,8 +420,6 @@ $ bDisableWaitOnMem : Disable UBA waiting on available memory before spawning ne
 
 $ bAllowKillOnMem : Let UBA kill running processes when close to out of memory
 
-$ bStoreObjFilesCompressed : Store object (.obj) compressed on disk. Requires uba to do link step where it will decompress obj files again
-
 $ OutputStatsThresholdMs : Threshold for when executor should output logging for the process. Defaults to never
 
 $ bWriteToDisk : Skip writing intermediate and output files to disk. Useful for validation builds where we don't need the output
diff --git a/Engine/Documentation/Source/Programming/UnrealBuildSystem/ModuleFiles/ModuleFilesProperties/ModuleFilesProperties.INT.udn b/Engine/Documentation/Source/Programming/UnrealBuildSystem/ModuleFiles/ModuleFilesProperties/ModuleFilesProperties.INT.udn
--- a/Engine/Documentation/Source/Programming/UnrealBuildSystem/ModuleFiles/ModuleFilesProperties/ModuleFilesProperties.INT.udn
+++ b/Engine/Documentation/Source/Programming/UnrealBuildSystem/ModuleFiles/ModuleFilesProperties/ModuleFilesProperties.INT.udn
@@ -204,7 +204,7 @@ $ bUsePrecompiled (Boolean): Whether this module should use precompiled data. Al
 
 $ bAllowConfidentialPlatformDefines (Boolean): Whether this module can use PLATFORM_XXXX style defines, where XXXX is a confidential platform name. This is used to ensure engine or other shared code does not reveal confidential information inside an #if PLATFORM_XXXX block. Licensee game code may want to allow for them, however.
 
-$ bAllowAutoRTFMInstrumentation (Boolean): Enables AutoRTFM instrumentation to this module only when AutoRTFMCompiler is enabled
+$ bDisableAutoRTFMInstrumentation (Boolean): Disables AutoRTFM instrumentation to this module only when AutoRTFMCompiler is enabled
 
 $ PrecompileForTargets (PrecompileTargetsType): Which targets this module should be precompiled for
 
diff --git a/Engine/Documentation/Source/Programming/UnrealBuildSystem/TargetFiles/TargetFilesProperties/TargetFilesProperties.INT.udn b/Engine/Documentation/Source/Programming/UnrealBuildSystem/TargetFiles/TargetFilesProperties/TargetFilesProperties.INT.udn
--- a/Engine/Documentation/Source/Programming/UnrealBuildSystem/TargetFiles/TargetFilesProperties/TargetFilesProperties.INT.udn
+++ b/Engine/Documentation/Source/Programming/UnrealBuildSystem/TargetFiles/TargetFilesProperties/TargetFilesProperties.INT.udn
@@ -198,6 +198,8 @@ $ bUseGameplayDebuggerCore (Boolean): Set to true when bUseGameplayDebugger is f
 
 $ bUseIris (Boolean): Whether to use Iris.
 
+$ bTrackRHIResourceInfo (Boolean): Whether to track owner (asset name) of RHI resource for non Shipping configurations. Useful for ListShaderMaps and ListShaderLibraries commands.
+
 $ bBuildEditor (Boolean): Whether we are compiling editor code or not. Prefer the more explicit bCompileAgainstEditor instead.
 
 $ bBuildRequiresCookedData (Boolean): Whether to compile code related to building assets. Consoles generally cannot build assets. Desktop platforms generally can.
