Documentation for build tools

diff --git a/Engine/Documentation/Source/Programming/UnrealBuildSystem/TargetFiles/TargetFilesProperties/TargetFilesProperties.INT.udn b/Engine/Documentation/Source/Programming/UnrealBuildSystem/TargetFiles/TargetFilesProperties/TargetFilesProperties.INT.udn
--- a/Engine/Documentation/Source/Programming/UnrealBuildSystem/TargetFiles/TargetFilesProperties/TargetFilesProperties.INT.udn
+++ b/Engine/Documentation/Source/Programming/UnrealBuildSystem/TargetFiles/TargetFilesProperties/TargetFilesProperties.INT.udn
@@ -198,7 +198,7 @@ $ bUseGameplayDebuggerCore (Boolean): Set to true when bUseGameplayDebugger is f
 
 $ bUseIris (Boolean): Whether to use Iris.
 
-$ bTrackRHIResourceInfo (Boolean): Whether to track owner (asset name) of RHI resource for non Shipping configurations. Useful for ListShaderMaps and ListShaderLibraries commands.
+$ bTrackRHIResourceInfoForTest (Boolean): Whether to track owner (asset name) of RHI resource for Test configuration. Useful for ListShaderMaps and ListShaderLibraries commands.
 
 $ bBuildEditor (Boolean): Whether we are compiling editor code or not. Prefer the more explicit bCompileAgainstEditor instead.
 
