Documentation for build tools

diff --git a/Engine/Documentation/Source/Programming/Development/BuildGraph/ScriptAnatomy/Tasks/PredefinedTasks/BuildGraphScriptPredefinedTasks.INT.udn b/Engine/Documentation/Source/Programming/Development/BuildGraph/ScriptAnatomy/Tasks/PredefinedTasks/BuildGraphScriptPredefinedTasks.INT.udn
--- a/Engine/Documentation/Source/Programming/Development/BuildGraph/ScriptAnatomy/Tasks/PredefinedTasks/BuildGraphScriptPredefinedTasks.INT.udn
+++ b/Engine/Documentation/Source/Programming/Development/BuildGraph/ScriptAnatomy/Tasks/PredefinedTasks/BuildGraphScriptPredefinedTasks.INT.udn
@@ -361,6 +361,16 @@ Replaces strings in a text file with secrets obtained from Horde
 | Replace | String    | Optional | Pairs of strings and secret names to expand in the text file, in the form SOURCE_TEXT=secret-name;SOURCE_TEXT_2=secret-name-2. If not specified, secrets embedded inline in the text will be expanded from {{secret-name.value}} strings. |
 | If      | Condition | Optional | Whether to execute this task. It is ignored if this condition evaluates to false.                                                                                                                                                         |
 
+### Horde-SetSecretEnvVar
+
+Replaces strings in a text file with secrets obtained from Horde
+
+|        |           |          |                                                                                   |
+| ------ | --------- | -------- | --------------------------------------------------------------------------------- |
+| Name   | String    | Required | Name of the environment variable to set                                           |
+| Secret | String    | Required | Name of the secret to fetch                                                       |
+| If     | Condition | Optional | Whether to execute this task. It is ignored if this condition evaluates to false. |
+
 ### Kubectl
 
 Spawns Kubectl and waits for it to complete.
diff --git a/Engine/Documentation/Source/Programming/UnrealBuildSystem/Configuration/BuildConfigProperties/BuildConfigProperties.INT.udn b/Engine/Documentation/Source/Programming/UnrealBuildSystem/Configuration/BuildConfigProperties/BuildConfigProperties.INT.udn
--- a/Engine/Documentation/Source/Programming/UnrealBuildSystem/Configuration/BuildConfigProperties/BuildConfigProperties.INT.udn
+++ b/Engine/Documentation/Source/Programming/UnrealBuildSystem/Configuration/BuildConfigProperties/BuildConfigProperties.INT.udn
@@ -64,7 +64,7 @@ $ bLegalToDistributeBinary : Whether the output from this target can be publicly
 
 $ bUseVerseBPVM : Whether to use the BPVM to run Verse.
 
-$ bUseAutoRTFMCompiler : Whether to use the AutoRTFM Clang compiler.
+$ bForceNoAutoRTFMCompiler : Whether to use force AutoRTFM Clang compiler off.
 
 $ bUseInlining : Enable inlining for all modules.
 
diff --git a/Engine/Documentation/Source/Programming/UnrealBuildSystem/TargetFiles/TargetFilesProperties/TargetFilesProperties.INT.udn b/Engine/Documentation/Source/Programming/UnrealBuildSystem/TargetFiles/TargetFilesProperties/TargetFilesProperties.INT.udn
--- a/Engine/Documentation/Source/Programming/UnrealBuildSystem/TargetFiles/TargetFilesProperties/TargetFilesProperties.INT.udn
+++ b/Engine/Documentation/Source/Programming/UnrealBuildSystem/TargetFiles/TargetFilesProperties/TargetFilesProperties.INT.udn
@@ -164,6 +164,8 @@ $ bUseVerseBPVM (Boolean): Whether to use the BPVM to run Verse.
 
 $ bUseAutoRTFMCompiler (Boolean): Whether to use the AutoRTFM Clang compiler.
 
+$ bForceNoAutoRTFMCompiler (Boolean): Whether to use force AutoRTFM Clang compiler off.
+
 $ bCompileChaos (Boolean): Whether to compile the Chaos physics plugin.
 
 $ bUseChaos (Boolean): Whether to use the Chaos physics interface. This overrides the physx flags to disable APEX and NvCloth
