Documentation for build tools

diff --git a/Engine/Documentation/Source/Programming/UnrealBuildSystem/TargetFiles/TargetFilesProperties/TargetFilesProperties.INT.udn b/Engine/Documentation/Source/Programming/UnrealBuildSystem/TargetFiles/TargetFilesProperties/TargetFilesProperties.INT.udn
--- a/Engine/Documentation/Source/Programming/UnrealBuildSystem/TargetFiles/TargetFilesProperties/TargetFilesProperties.INT.udn
+++ b/Engine/Documentation/Source/Programming/UnrealBuildSystem/TargetFiles/TargetFilesProperties/TargetFilesProperties.INT.udn
@@ -232,6 +232,8 @@ $ bCompileAgainstApplicationCore (Boolean): Enabled for builds that need to init
 
 $ bEnableTrace (Boolean): Enabled for builds that want to use against the Trace module for profiling and diagnostics.
 
+$ bForceEnableTrace (Boolean): Force enable tracing - used to allow testing programs to verify tracing works as intended.
+
 $ bCompileAgainstEditor (Boolean): Enabled for editor builds (TargetType.Editor). Can be overridden for programs (TargetType.Program) that would need to compile against editor code. Not available for other target types. Mainly drives the value of WITH_EDITOR.
 
 $ bCompileRecast (Boolean): Whether to compile Recast navmesh generation.
