Documentation for build tools

diff --git a/Engine/Documentation/Source/Programming/UnrealBuildSystem/Configuration/BuildConfigProperties/BuildConfigProperties.INT.udn b/Engine/Documentation/Source/Programming/UnrealBuildSystem/Configuration/BuildConfigProperties/BuildConfigProperties.INT.udn
--- a/Engine/Documentation/Source/Programming/UnrealBuildSystem/Configuration/BuildConfigProperties/BuildConfigProperties.INT.udn
+++ b/Engine/Documentation/Source/Programming/UnrealBuildSystem/Configuration/BuildConfigProperties/BuildConfigProperties.INT.udn
@@ -330,6 +330,8 @@ $ bStrictEnumTypesConformance : Enables enum types conformance (/Zc:enumTypes) i
 
 $ bStripPrivateSymbols : Whether to request the linker create a stripped pdb file as part of the build. If enabled the full debug pdb will have the extension .full.pdb
 
+$ bNoLinkerDebugInfo : If you supply -NoDebugInfo, windows platforms still create debug info while linking. Set this to true to not create debug info while linking in this circumstance
+
 $ PCHMemoryAllocationFactor : Determines the amount of memory that the compiler allocates to construct precompiled headers (/Zm).
 
 $ AdditionalLinkerOptions : Allow the target to specify extra options for linking that aren't otherwise noted here
