Documentation for build tools

diff --git a/Engine/Documentation/Source/Programming/UnrealBuildSystem/Configuration/BuildConfigProperties/BuildConfigProperties.INT.udn b/Engine/Documentation/Source/Programming/UnrealBuildSystem/Configuration/BuildConfigProperties/BuildConfigProperties.INT.udn
--- a/Engine/Documentation/Source/Programming/UnrealBuildSystem/Configuration/BuildConfigProperties/BuildConfigProperties.INT.udn
+++ b/Engine/Documentation/Source/Programming/UnrealBuildSystem/Configuration/BuildConfigProperties/BuildConfigProperties.INT.udn
@@ -98,6 +98,8 @@ $ UnsafeTypeCastWarningLevel : Indicates what warning/error level to treat unsaf
 
 $ bUndefinedIdentifierErrors : Forces the use of undefined identifiers in conditional expressions to be treated as errors.
 
+$ UndefinedIdentifierWarningLevel : Indicates what warning/error level to treat undefined identifiers in conditional expressions.
+
 $ PCHPerformanceIssueWarningLevel : Indicates what warning/error level to treat potential PCH performance issues.
 
 $ ModuleIncludePathWarningLevel : How to treat general module include path validation messages
diff --git a/Engine/Documentation/Source/Programming/UnrealBuildSystem/ModuleFiles/ModuleFilesProperties/ModuleFilesProperties.INT.udn b/Engine/Documentation/Source/Programming/UnrealBuildSystem/ModuleFiles/ModuleFilesProperties/ModuleFilesProperties.INT.udn
--- a/Engine/Documentation/Source/Programming/UnrealBuildSystem/ModuleFiles/ModuleFilesProperties/ModuleFilesProperties.INT.udn
+++ b/Engine/Documentation/Source/Programming/UnrealBuildSystem/ModuleFiles/ModuleFilesProperties/ModuleFilesProperties.INT.udn
@@ -160,6 +160,8 @@ $ bWarningsAsErrors (Boolean): Whether to enable all warnings as errors. UE enab
 
 $ UnsafeTypeCastWarningLevel (WarningLevel): How to treat unsafe implicit type cast warnings (e.g., double->float or int64->int32)
 
+$ UndefinedIdentifierWarningLevel (WarningLevel): Indicates what warning/error level to treat undefined identifiers in conditional expressions.
+
 $ bEnableUndefinedIdentifierWarnings (Boolean): Enable warnings for using undefined identifiers in #if expressions
 
 $ ModuleIncludePathWarningLevel (WarningLevel): How to treat general module include path validation messages
diff --git a/Engine/Documentation/Source/Programming/UnrealBuildSystem/TargetFiles/TargetFilesProperties/TargetFilesProperties.INT.udn b/Engine/Documentation/Source/Programming/UnrealBuildSystem/TargetFiles/TargetFilesProperties/TargetFilesProperties.INT.udn
--- a/Engine/Documentation/Source/Programming/UnrealBuildSystem/TargetFiles/TargetFilesProperties/TargetFilesProperties.INT.udn
+++ b/Engine/Documentation/Source/Programming/UnrealBuildSystem/TargetFiles/TargetFilesProperties/TargetFilesProperties.INT.udn
@@ -386,6 +386,8 @@ $ UnsafeTypeCastWarningLevel (WarningLevel): Indicates what warning/error level
 
 $ bUndefinedIdentifierErrors (Boolean): Forces the use of undefined identifiers in conditional expressions to be treated as errors.
 
+$ UndefinedIdentifierWarningLevel (WarningLevel): Indicates what warning/error level to treat undefined identifiers in conditional expressions.
+
 $ PCHPerformanceIssueWarningLevel (WarningLevel): Indicates what warning/error level to treat potential PCH performance issues.
 
 $ ModuleIncludePathWarningLevel (WarningLevel): How to treat general module include path validation messages
