Documentation for build tools

diff --git a/Engine/Documentation/Source/Programming/UnrealBuildSystem/ModuleFiles/ModuleFilesProperties/ModuleFilesProperties.INT.udn b/Engine/Documentation/Source/Programming/UnrealBuildSystem/ModuleFiles/ModuleFilesProperties/ModuleFilesProperties.INT.udn
--- a/Engine/Documentation/Source/Programming/UnrealBuildSystem/ModuleFiles/ModuleFilesProperties/ModuleFilesProperties.INT.udn
+++ b/Engine/Documentation/Source/Programming/UnrealBuildSystem/ModuleFiles/ModuleFilesProperties/ModuleFilesProperties.INT.udn
@@ -12,6 +12,8 @@ $ Logger (ILogger): Accessor for the target logger
 
 $ HasPackageOverride (Boolean): Returns true if there has been an override type specified on this module
 
+$ ForceIncludeFiles (List&lt;String&gt;): Header files that should be force included for every source file in this module.
+
 $ SDKVersionRelevantPlatforms (List&lt;UnrealTargetPlatform&gt;): If this contains a platform and the project has had it's SDK version overriden from default, this module will be compiled as a project module, instead of shared engine module
 
 $ StaticAnalyzerCheckers (HashSet&lt;String&gt;): The static analyzer checkers that should be enabled rather than the defaults. This is only supported for Clang. See https://clang.llvm.org/docs/analyzer/checkers.html for a full list. Or run: 'clang -Xclang -analyzer-checker-help' or: 'clang -Xclang -analyzer-checker-help-alpha' for the list of experimental checkers.
diff --git a/Engine/Documentation/Source/Programming/UnrealBuildSystem/TargetFiles/TargetFilesProperties/TargetFilesProperties.INT.udn b/Engine/Documentation/Source/Programming/UnrealBuildSystem/TargetFiles/TargetFilesProperties/TargetFilesProperties.INT.udn
--- a/Engine/Documentation/Source/Programming/UnrealBuildSystem/TargetFiles/TargetFilesProperties/TargetFilesProperties.INT.udn
+++ b/Engine/Documentation/Source/Programming/UnrealBuildSystem/TargetFiles/TargetFilesProperties/TargetFilesProperties.INT.udn
@@ -148,6 +148,8 @@ $ SolutionDirectory (String): Allows a Program Target to specify it's own soluti
 
 $ bGenerateProgramProject (Boolean): Force a Target to be treated as a Program for the purposes of project file generation.
 
+$ bExplicitTargetForType (Boolean): If true, GetTargetNameByType will not consider this target, thereby disambiguating -TargetType=X.
+
 $ bBuildInSolutionByDefault (Nullable&lt;Boolean&gt;): Whether the target should be included in the default solution build configuration Setting this to false will skip building when running in the IDE
 
 $ bShouldCompileAsDLL (Boolean): Whether this target should be compiled as a DLL.  Requires LinkType to be set to TargetLinkType.Monolithic.
