Documentation for build tools

diff --git a/Engine/Documentation/Source/Programming/UnrealBuildSystem/TargetFiles/TargetFilesProperties/TargetFilesProperties.INT.udn b/Engine/Documentation/Source/Programming/UnrealBuildSystem/TargetFiles/TargetFilesProperties/TargetFilesProperties.INT.udn
--- a/Engine/Documentation/Source/Programming/UnrealBuildSystem/TargetFiles/TargetFilesProperties/TargetFilesProperties.INT.udn
+++ b/Engine/Documentation/Source/Programming/UnrealBuildSystem/TargetFiles/TargetFilesProperties/TargetFilesProperties.INT.udn
@@ -200,6 +200,8 @@ $ bUseGameplayDebugger (Boolean): Whether to compile with WITH_GAMEPLAY_DEBUGGER
 
 $ bUseGameplayDebuggerCore (Boolean): Set to true when bUseGameplayDebugger is false but GameplayDebugger's core parts are required.
 
+$ bCompileIoStoreOnDemand (Boolean): Whether to use I/O store on-demand
+
 $ bUseIris (Boolean): Whether to use Iris.
 
 $ bTrackRHIResourceInfoForTest (Boolean): Whether to track owner (asset name) of RHI resource for Test configuration. Useful for ListShaderMaps and ListShaderLibraries commands.
