Documentation for build tools

diff --git a/Engine/Documentation/Source/Programming/UnrealBuildSystem/Configuration/BuildConfigProperties/BuildConfigProperties.INT.udn b/Engine/Documentation/Source/Programming/UnrealBuildSystem/Configuration/BuildConfigProperties/BuildConfigProperties.INT.udn
--- a/Engine/Documentation/Source/Programming/UnrealBuildSystem/Configuration/BuildConfigProperties/BuildConfigProperties.INT.udn
+++ b/Engine/Documentation/Source/Programming/UnrealBuildSystem/Configuration/BuildConfigProperties/BuildConfigProperties.INT.udn
@@ -532,6 +532,8 @@ $ WindowsPool : Pool for the Horde agent to assign when on Windows
 
 $ Requirements : Requirements for the Horde agent to assign
 
+$ Cluster : Compute cluster ID to use in Horde. Set to "_auto" to let Horde server resolve a suitable cluster. In multi-region setups this is can simplify configuration of UBT/UBA a lot.
+
 $ LocalHost : Which ip UBA server should give to agents. This will invert so host listens and agents connect
 
 $ ConnectionMode : Connection mode for agent/compute communication for valid modes.
