Documentation for build tools

diff --git a/Engine/Documentation/Source/Programming/Development/BuildGraph/ScriptAnatomy/Tasks/PredefinedTasks/BuildGraphScriptPredefinedTasks.INT.udn b/Engine/Documentation/Source/Programming/Development/BuildGraph/ScriptAnatomy/Tasks/PredefinedTasks/BuildGraphScriptPredefinedTasks.INT.udn
--- a/Engine/Documentation/Source/Programming/Development/BuildGraph/ScriptAnatomy/Tasks/PredefinedTasks/BuildGraphScriptPredefinedTasks.INT.udn
+++ b/Engine/Documentation/Source/Programming/Development/BuildGraph/ScriptAnatomy/Tasks/PredefinedTasks/BuildGraphScriptPredefinedTasks.INT.udn
@@ -508,6 +508,24 @@ Creates files containing random data in the specified output directory. Used for
 | Tag       | String    | Optional | Optional filter to be applied to the list of input files.                         |
 | If        | Condition | Optional | Whether to execute this task. It is ignored if this condition evaluates to false. |
 
+### Reconcile
+
+Creates a new changelist and reconciles a set of files to submit to a Perforce stream.
+
+|             |                    |          |                                                                                                                                                                              |
+| ----------- | ------------------ | -------- | ---------------------------------------------------------------------------------------------------------------------------------------------------------------------------- |
+| Description | String             | Required | The description for the submitted changelist.                                                                                                                                |
+| Files       | FileSpec           | Optional | The files to reconcile                                                                                                                                                       |
+| Directories | String             | Optional | The directories to reconcile, semi colon delimited, relative p4 syntax                                                                                                       |
+| Workspace   | String             | Optional | The workspace name. If specified, a new workspace will be created using the given stream and root directory to submit the files. If not, the current workspace will be used. |
+| Stream      | String             | Optional | The stream for the workspace -- defaults to the current stream. Ignored unless the Workspace attribute is also specified.                                                    |
+| Branch      | String             | Optional | Branch for the workspace (legacy P4 depot path). May not be used in conjunction with Stream.                                                                                 |
+| RootDir     | DirectoryReference | Optional | Root directory for the stream. If not specified, defaults to the current root directory.                                                                                     |
+| Force       | Boolean            | Optional | Force the submit to happen -- even if a resolve is needed (always accept current version).                                                                                   |
+| P4Verbose   | Boolean            | Optional | Allow verbose P4 output (spew).                                                                                                                                              |
+| Preview     | Boolean            | Required | Runs a reconcile preview, does not submit.                                                                                                                                   |
+| If          | Condition          | Optional | Whether to execute this task. It is ignored if this condition evaluates to false.                                                                                            |
+
 ### Rename
 
 Renames a file, or group of files.
