Documentation for build tools

diff --git a/Engine/Documentation/Source/Programming/Development/BuildGraph/ScriptAnatomy/Tasks/PredefinedTasks/BuildGraphScriptPredefinedTasks.INT.udn b/Engine/Documentation/Source/Programming/Development/BuildGraph/ScriptAnatomy/Tasks/PredefinedTasks/BuildGraphScriptPredefinedTasks.INT.udn
--- a/Engine/Documentation/Source/Programming/Development/BuildGraph/ScriptAnatomy/Tasks/PredefinedTasks/BuildGraphScriptPredefinedTasks.INT.udn
+++ b/Engine/Documentation/Source/Programming/Development/BuildGraph/ScriptAnatomy/Tasks/PredefinedTasks/BuildGraphScriptPredefinedTasks.INT.udn
@@ -800,6 +800,15 @@ Imports an oplog from Zen to a specified destination.
 | Key         | String        | Optional | What key to use when importing from cloud                                                          |
 | If          | Condition     | Optional | Whether to execute this task. It is ignored if this condition evaluates to false.                  |
 
+### ZenLaunch
+
+Launches ZenServer
+
+|         |               |          |                                                                                   |
+| ------- | ------------- | -------- | --------------------------------------------------------------------------------- |
+| Project | FileReference | Required | The project for which to launch ZenServer                                         |
+| If      | Condition     | Optional | Whether to execute this task. It is ignored if this condition evaluates to false. |
+
 ### Zip
 
 Compresses files into a zip archive.
