Documentation for build tools

diff --git a/Engine/Documentation/Source/Programming/UnrealBuildSystem/Configuration/BuildConfigProperties/BuildConfigProperties.INT.udn b/Engine/Documentation/Source/Programming/UnrealBuildSystem/Configuration/BuildConfigProperties/BuildConfigProperties.INT.udn
--- a/Engine/Documentation/Source/Programming/UnrealBuildSystem/Configuration/BuildConfigProperties/BuildConfigProperties.INT.udn
+++ b/Engine/Documentation/Source/Programming/UnrealBuildSystem/Configuration/BuildConfigProperties/BuildConfigProperties.INT.udn
@@ -66,6 +66,8 @@ $ bUseVerseBPVM : Whether to use the BPVM to run Verse.
 
 $ bForceNoAutoRTFMCompiler : Whether to use force AutoRTFM Clang compiler off.
 
+$ bUseAutoRTFMVerifier : Whether to enable emitting AutoRTFM verification metadata
+
 $ bUseInlining : Enable inlining for all modules.
 
 $ bUseDebugLiveCodingConsole : Whether to enable support for live coding
diff --git a/Engine/Documentation/Source/Programming/UnrealBuildSystem/TargetFiles/TargetFilesProperties/TargetFilesProperties.INT.udn b/Engine/Documentation/Source/Programming/UnrealBuildSystem/TargetFiles/TargetFilesProperties/TargetFilesProperties.INT.udn
--- a/Engine/Documentation/Source/Programming/UnrealBuildSystem/TargetFiles/TargetFilesProperties/TargetFilesProperties.INT.udn
+++ b/Engine/Documentation/Source/Programming/UnrealBuildSystem/TargetFiles/TargetFilesProperties/TargetFilesProperties.INT.udn
@@ -168,6 +168,8 @@ $ bUseAutoRTFMCompiler (Boolean): Whether to use the AutoRTFM Clang compiler.
 
 $ bForceNoAutoRTFMCompiler (Boolean): Whether to use force AutoRTFM Clang compiler off.
 
+$ bUseAutoRTFMVerifier (Boolean): Whether to enable emitting AutoRTFM verification metadata
+
 $ bCompileChaos (Boolean): Whether to compile the Chaos physics plugin.
 
 $ bUseChaos (Boolean): Whether to use the Chaos physics interface. This overrides the physx flags to disable APEX and NvCloth
