Documentation for build tools

diff --git a/Engine/Documentation/Source/Programming/Development/BuildGraph/ScriptAnatomy/Tasks/PredefinedTasks/BuildGraphScriptPredefinedTasks.INT.udn b/Engine/Documentation/Source/Programming/Development/BuildGraph/ScriptAnatomy/Tasks/PredefinedTasks/BuildGraphScriptPredefinedTasks.INT.udn
--- a/Engine/Documentation/Source/Programming/Development/BuildGraph/ScriptAnatomy/Tasks/PredefinedTasks/BuildGraphScriptPredefinedTasks.INT.udn
+++ b/Engine/Documentation/Source/Programming/Development/BuildGraph/ScriptAnatomy/Tasks/PredefinedTasks/BuildGraphScriptPredefinedTasks.INT.udn
@@ -809,6 +809,17 @@ Launches ZenServer
 | Project | FileReference | Required | The project for which to launch ZenServer                                         |
 | If      | Condition     | Optional | Whether to execute this task. It is ignored if this condition evaluates to false. |
 
+### ZenMirror
+
+Exports an snapshot from Zen to a specified destination.
+
+|                    |                    |          |                                                                                                                                                      |
+| ------------------ | ------------------ | -------- | ---------------------------------------------------------------------------------------------------------------------------------------------------- |
+| Project            | FileReference      | Optional | The project from which to export the snapshot                                                                                                        |
+| Platform           | String             | Optional | The target platform to mirror the snapshot for                                                                                                       |
+| DestinationFileDir | DirectoryReference | Optional | The path on the local disk to which the data will be mirrored If empty then the path will be set to the %Project%\Saved\Cooked\%Platform% directory. |
+| If                 | Condition          | Optional | Whether to execute this task. It is ignored if this condition evaluates to false.                                                                    |
+
 ### Zip
 
 Compresses files into a zip archive.
