Documentation for build tools

diff --git a/Engine/Documentation/Source/Programming/UnrealBuildSystem/Configuration/BuildConfigProperties/BuildConfigProperties.INT.udn b/Engine/Documentation/Source/Programming/UnrealBuildSystem/Configuration/BuildConfigProperties/BuildConfigProperties.INT.udn
--- a/Engine/Documentation/Source/Programming/UnrealBuildSystem/Configuration/BuildConfigProperties/BuildConfigProperties.INT.udn
+++ b/Engine/Documentation/Source/Programming/UnrealBuildSystem/Configuration/BuildConfigProperties/BuildConfigProperties.INT.udn
@@ -260,8 +260,6 @@ $ bTuneDebugInfoForLLDB : Turns on tuning of debug info for LLDB
 
 $ bDisableDumpSyms : Whether to globally disable calling dump_syms
 
-$ bWriteSarif : Whether .sarif files containing errors and warnings are written alongside each .obj, if supported
-
 $ bWriteSolutionOptionFile : Whether to write a solution option (suo) file for the sln.
 
 $ bVsConfigFile : Whether to write a .vsconfig file next to the sln to suggest components to install.
@@ -450,12 +448,6 @@ $ CacheMaxWorkers : Max number of cache download tasks that can execute in paral
 
 $ ReportCacheMissReason : Report reason a cache miss happened. Useful when searching for determinism/portability issues
 
-$ HordeMaxCores : Max cores allowed to be used by build session
-
-$ HordeDelay : How long UBT should wait to ask for help. Useful in build configs where machine can delay remote work and still get same wall time results (pch dependencies etc)
-
-$ bHordeAllowWine : Allow use of Wine. Only applicable to Horde agents running Linux. Can still be ignored if Wine executable is not set on agent.
-
 $ bDisableHorde : Disable horde all together
 
 ### TestTargetRules
@@ -538,6 +530,12 @@ $ Cluster : Compute cluster ID to use in Horde. Set to "_auto" to let Horde serv
 
 $ LocalHost : Which ip UBA server should give to agents. This will invert so host listens and agents connect
 
+$ MaxCores : Max cores allowed to be used by build session
+
+$ StartupDelay : How long UBT should wait to ask for help. Useful in build configs where machine can delay remote work and still get same wall time results (pch dependencies etc)
+
+$ AllowWine : Allow use of Wine. Only applicable to Horde agents running Linux. Can still be ignored if Wine executable is not set on agent.
+
 $ ConnectionMode : Connection mode for agent/compute communication for valid modes.
 
 $ Encryption : Encryption to use for agent/compute communication. Note that UBA agent uses its own encryption. for valid modes.
