Documentation for build tools

diff --git a/Engine/Documentation/Source/Programming/Development/BuildGraph/ScriptAnatomy/Tasks/PredefinedTasks/BuildGraphScriptPredefinedTasks.INT.udn b/Engine/Documentation/Source/Programming/Development/BuildGraph/ScriptAnatomy/Tasks/PredefinedTasks/BuildGraphScriptPredefinedTasks.INT.udn
--- a/Engine/Documentation/Source/Programming/Development/BuildGraph/ScriptAnatomy/Tasks/PredefinedTasks/BuildGraphScriptPredefinedTasks.INT.udn
+++ b/Engine/Documentation/Source/Programming/Development/BuildGraph/ScriptAnatomy/Tasks/PredefinedTasks/BuildGraphScriptPredefinedTasks.INT.udn
@@ -144,18 +144,20 @@ Copies files from one directory to another.
 
 ### CreateArtifact
 
-Deploys a tool update through Horde
-
-|             |           |          |                                                                                   |
-| ----------- | --------- | -------- | --------------------------------------------------------------------------------- |
-| Name        | String    | Required | Name for the artifact                                                             |
-| Type        | String    | Optional | Type of the artifact                                                              |
-| Description | String    | Optional | Description for the artifact                                                      |
-| BaseDir     | String    | Required | Base directory to resolve relative paths for input files.                         |
-| StreamId    | String    | Optional | Stream for this artifact                                                          |
-| Change      | Integer   | Optional | Changelist number for this artifact                                               |
-| Files       | FileSpec  | Optional | Files to be uploaded.                                                             |
-| If          | Condition | Optional | Whether to execute this task. It is ignored if this condition evaluates to false. |
+Uploads an artifact to Horde
+
+|             |           |          |                                                                                                                                |
+| ----------- | --------- | -------- | ------------------------------------------------------------------------------------------------------------------------------ |
+| Name        | String    | Required | Name of the artifact                                                                                                           |
+| Type        | String    | Required | The artifact type. Determines the permissions and expiration policy for the artifact.                                          |
+| Description | String    | Optional | Description for the artifact. Will be shown through the Horde dashboard.                                                       |
+| BaseDir     | String    | Optional | Base path for the uploaded files. All the tagged files must be under this directory. Defaults to the workspace root directory. |
+| StreamId    | String    | Optional | Stream containing the artifact.                                                                                                |
+| Commit      | String    | Optional | Commit for the uploaded artifact.                                                                                              |
+| Files       | FileSpec  | Required | Files to include in the artifact.                                                                                              |
+| Keys        | String    | Optional | Queryable keys for this artifact, separated by semicolons.                                                                     |
+| Metadata    | String    | Optional | Other metadata for the artifact, separated by semicolons.                                                                      |
+| If          | Condition | Optional | Whether to execute this task. It is ignored if this condition evaluates to false.                                              |
 
 ### CsCompile
 
