Documentation for build tools

diff --git a/Engine/Documentation/Source/Programming/Development/BuildGraph/ScriptAnatomy/Tasks/PredefinedTasks/BuildGraphScriptPredefinedTasks.INT.udn b/Engine/Documentation/Source/Programming/Development/BuildGraph/ScriptAnatomy/Tasks/PredefinedTasks/BuildGraphScriptPredefinedTasks.INT.udn
--- a/Engine/Documentation/Source/Programming/Development/BuildGraph/ScriptAnatomy/Tasks/PredefinedTasks/BuildGraphScriptPredefinedTasks.INT.udn
+++ b/Engine/Documentation/Source/Programming/Development/BuildGraph/ScriptAnatomy/Tasks/PredefinedTasks/BuildGraphScriptPredefinedTasks.INT.udn
@@ -157,6 +157,7 @@ Uploads an artifact to Horde
 | Files       | FileSpec  | Required | Files to include in the artifact.                                                                                              |
 | Keys        | String    | Optional | Queryable keys for this artifact, separated by semicolons.                                                                     |
 | Metadata    | String    | Optional | Other metadata for the artifact, separated by semicolons.                                                                      |
+| Symbols     | Boolean   | Optional | Whether to add aliases for symbol files                                                                                        |
 | If          | Condition | Optional | Whether to execute this task. It is ignored if this condition evaluates to false.                                              |
 
 ### CsCompile
