Documentation for build tools

diff --git a/Engine/Documentation/Source/Programming/UnrealBuildSystem/TargetFiles/TargetFilesProperties/TargetFilesProperties.INT.udn b/Engine/Documentation/Source/Programming/UnrealBuildSystem/TargetFiles/TargetFilesProperties/TargetFilesProperties.INT.udn
--- a/Engine/Documentation/Source/Programming/UnrealBuildSystem/TargetFiles/TargetFilesProperties/TargetFilesProperties.INT.udn
+++ b/Engine/Documentation/Source/Programming/UnrealBuildSystem/TargetFiles/TargetFilesProperties/TargetFilesProperties.INT.udn
@@ -28,7 +28,9 @@ $ DebugInfoLineTablesOnlyPlugins (HashSet&lt;String&gt;): Plugins that should em
 
 $ MergePlugins (List&lt;String&gt;): Experimental: List of plugins (and their dependencies) to each merge into separate libraries. Requires bMergeModules to be enabled
 
-$ MergePluginsShared (Dictionary&lt;String, HashSet&lt;String&gt;&gt;): Experimental: List of merged plugins to move common dependenceis to a shared library, requires MergePlugins. Can be chained
+$ MergePluginsShared (Dictionary&lt;String, HashSet&lt;String&gt;&gt;): Experimental: List of merged plugins to move common dependencies to a shared library, requires MergePlugins. Can be chained
+
+$ MergePluginsLaunch (List&lt;String&gt;): Experimental: List of merged plugins to move to the main executable, requires MergePlugins. "Engine" for all shared engine modules and "Common" for shared project modules.
 
 $ GlobalDefinitions (List&lt;String&gt;): Macros to define globally across the whole target.
 
