Documentation for build tools

diff --git a/Engine/Documentation/Source/Programming/UnrealBuildSystem/TargetFiles/TargetFilesProperties/TargetFilesProperties.INT.udn b/Engine/Documentation/Source/Programming/UnrealBuildSystem/TargetFiles/TargetFilesProperties/TargetFilesProperties.INT.udn
--- a/Engine/Documentation/Source/Programming/UnrealBuildSystem/TargetFiles/TargetFilesProperties/TargetFilesProperties.INT.udn
+++ b/Engine/Documentation/Source/Programming/UnrealBuildSystem/TargetFiles/TargetFilesProperties/TargetFilesProperties.INT.udn
@@ -132,6 +132,8 @@ $ bLegalToDistributeBinary (Boolean): Whether the output from this target can be
 
 $ UndecoratedConfiguration (UnrealTargetConfiguration): Specifies the configuration whose binaries do not require a "-Platform-Configuration" suffix.
 
+$ DecoratedSeparator (String): Specifies the separator charcter for binary filenames.
+
 $ bAllowHotReload (Boolean): Whether this target supports hot reload
 
 $ bBuildAllModules (Boolean): Build all the modules that are valid for this target type. Used for CIS and making installed engine builds.
