Documentation for build tools

diff --git a/Engine/Documentation/Source/Programming/UnrealBuildSystem/Configuration/BuildConfigProperties/BuildConfigProperties.INT.udn b/Engine/Documentation/Source/Programming/UnrealBuildSystem/Configuration/BuildConfigProperties/BuildConfigProperties.INT.udn
--- a/Engine/Documentation/Source/Programming/UnrealBuildSystem/Configuration/BuildConfigProperties/BuildConfigProperties.INT.udn
+++ b/Engine/Documentation/Source/Programming/UnrealBuildSystem/Configuration/BuildConfigProperties/BuildConfigProperties.INT.udn
@@ -222,7 +222,7 @@ $ CppStandard : Which C++ standard to use for compiling this target (for non-eng
 
 $ CStandard : Which C standard to use for compiling this target
 
-$ MinCpuArchX64 : Direct the compiler to generate AVX instructions wherever SSE or AVX intrinsics are used, on the x64 platforms that support it. Note that by enabling this you are changing the minspec for the PC platform, and the resultant executable will crash on machines without AVX support.
+$ MinCpuArchX64 : Direct the compiler to generate AVX instructions wherever SSE or AVX intrinsics are used, on the x64 platforms that support it. Ignored for arm64. Note that by enabling this you are changing the minspec for the target platform, and the resultant executable will crash on machines without AVX support.
 
 $ ActionStallReportTime : Number of second of no completed actions to trigger an action stall report. If zero, stall reports will not be enabled.
 
@@ -330,6 +330,8 @@ $ bStrictPreprocessorConformance : Enables new preprocessor conformance (/Zc:pre
 
 $ bStrictEnumTypesConformance : Enables enum types conformance (/Zc:enumTypes) in VS2022 17.4 Preview 4.0+.
 
+$ bStrictODRViolationConformance : Enables enforcing standard C++ ODR violations (/Zc:checkGwOdr) in VS2022 17.5 Preview 2.0+
+
 $ bStripPrivateSymbols : Whether to request the linker create a stripped pdb file as part of the build. If enabled the full debug pdb will have the extension .full.pdb
 
 $ bNoLinkerDebugInfo : If you supply -NoDebugInfo, windows platforms still create debug info while linking. Set this to true to not create debug info while linking in this circumstance
diff --git a/Engine/Documentation/Source/Programming/UnrealBuildSystem/TargetFiles/TargetFilesProperties/TargetFilesProperties.INT.udn b/Engine/Documentation/Source/Programming/UnrealBuildSystem/TargetFiles/TargetFilesProperties/TargetFilesProperties.INT.udn
--- a/Engine/Documentation/Source/Programming/UnrealBuildSystem/TargetFiles/TargetFilesProperties/TargetFilesProperties.INT.udn
+++ b/Engine/Documentation/Source/Programming/UnrealBuildSystem/TargetFiles/TargetFilesProperties/TargetFilesProperties.INT.udn
@@ -566,7 +566,7 @@ $ CppStandard (CppStandardVersion): Which C++ standard to use for compiling this
 
 $ CStandard (CStandardVersion): Which C standard to use for compiling this target
 
-$ MinCpuArchX64 (MinimumCpuArchitectureX64): Direct the compiler to generate AVX instructions wherever SSE or AVX intrinsics are used, on the x64 platforms that support it. Note that by enabling this you are changing the minspec for the PC platform, and the resultant executable will crash on machines without AVX support.
+$ MinCpuArchX64 (MinimumCpuArchitectureX64): Direct the compiler to generate AVX instructions wherever SSE or AVX intrinsics are used, on the x64 platforms that support it. Ignored for arm64. Note that by enabling this you are changing the minspec for the target platform, and the resultant executable will crash on machines without AVX support.
 
 $ BuildVersion (String): The build version string
 
