Documentation for build tools

diff --git a/Engine/Documentation/Source/Programming/UnrealBuildSystem/TargetFiles/TargetFilesProperties/TargetFilesProperties.INT.udn b/Engine/Documentation/Source/Programming/UnrealBuildSystem/TargetFiles/TargetFilesProperties/TargetFilesProperties.INT.udn
--- a/Engine/Documentation/Source/Programming/UnrealBuildSystem/TargetFiles/TargetFilesProperties/TargetFilesProperties.INT.udn
+++ b/Engine/Documentation/Source/Programming/UnrealBuildSystem/TargetFiles/TargetFilesProperties/TargetFilesProperties.INT.udn
@@ -486,6 +486,8 @@ $ bSupportEditAndContinue (Boolean): Whether to support edit and continue.
 
 $ bOmitFramePointers (Boolean): Whether to omit frame pointers or not. Disabling is useful for e.g. memory profiling on the PC.
 
+$ bCompressDebugFile (Boolean): If we want to compress the debug files
+
 $ bEnableCppModules (Boolean): Whether to enable support for C++20 modules
 
 $ bEnableCppCoroutinesForEvaluation (Boolean): Whether to enable support for C++20 coroutines This option is provided to facilitate evaluation of the feature. Expect the name of the option to change. Use of coroutines in with UE is untested and unsupported.
