Documentation for build tools

diff --git a/Engine/Documentation/Source/Programming/Development/BuildGraph/ScriptAnatomy/Tasks/PredefinedTasks/BuildGraphScriptPredefinedTasks.INT.udn b/Engine/Documentation/Source/Programming/Development/BuildGraph/ScriptAnatomy/Tasks/PredefinedTasks/BuildGraphScriptPredefinedTasks.INT.udn
--- a/Engine/Documentation/Source/Programming/Development/BuildGraph/ScriptAnatomy/Tasks/PredefinedTasks/BuildGraphScriptPredefinedTasks.INT.udn
+++ b/Engine/Documentation/Source/Programming/Development/BuildGraph/ScriptAnatomy/Tasks/PredefinedTasks/BuildGraphScriptPredefinedTasks.INT.udn
@@ -158,6 +158,7 @@ Uploads an artifact to Horde
 | Keys        | String    | Optional | Queryable keys for this artifact, separated by semicolons.                                                                     |
 | Metadata    | String    | Optional | Other metadata for the artifact, separated by semicolons.                                                                      |
 | Symbols     | Boolean   | Optional | Whether to add aliases for symbol files                                                                                        |
+| Dedupe      | Boolean   | Optional | Whether to reuse blobs from the previous artifact                                                                              |
 | If          | Condition | Optional | Whether to execute this task. It is ignored if this condition evaluates to false.                                              |
 
 ### CsCompile
@@ -506,6 +507,7 @@ Creates files containing random data in the specified output directory. Used for
 
 |           |           |          |                                                                                   |
 | --------- | --------- | -------- | --------------------------------------------------------------------------------- |
+| Seed      | Integer   | Optional | Seed for the data generation                                                      |
 | Size      | Integer   | Optional | The size of each file.                                                            |
 | Count     | Integer   | Optional | Number of files to write.                                                         |
 | Different | Boolean   | Optional | Whether to generate different data for each output file.                          |
diff --git a/Engine/Documentation/Source/Programming/UnrealBuildSystem/Configuration/BuildConfigProperties/BuildConfigProperties.INT.udn b/Engine/Documentation/Source/Programming/UnrealBuildSystem/Configuration/BuildConfigProperties/BuildConfigProperties.INT.udn
--- a/Engine/Documentation/Source/Programming/UnrealBuildSystem/Configuration/BuildConfigProperties/BuildConfigProperties.INT.udn
+++ b/Engine/Documentation/Source/Programming/UnrealBuildSystem/Configuration/BuildConfigProperties/BuildConfigProperties.INT.udn
@@ -68,6 +68,8 @@ $ bForceNoAutoRTFMCompiler : Whether to use force AutoRTFM Clang compiler off.
 
 $ bUseAutoRTFMVerifier : Whether to enable emitting AutoRTFM verification metadata
 
+$ bAutoRTFMVerify : Whether to run LLVM verification after the AutoRTFM compiler pass. This is used by our compiler folks to ensure the pass works with the various code-paths UBT can take us down.
+
 $ bUseInlining : Enable inlining for all modules.
 
 $ bUseDebugLiveCodingConsole : Whether to enable support for live coding
diff --git a/Engine/Documentation/Source/Programming/UnrealBuildSystem/TargetFiles/TargetFilesProperties/TargetFilesProperties.INT.udn b/Engine/Documentation/Source/Programming/UnrealBuildSystem/TargetFiles/TargetFilesProperties/TargetFilesProperties.INT.udn
--- a/Engine/Documentation/Source/Programming/UnrealBuildSystem/TargetFiles/TargetFilesProperties/TargetFilesProperties.INT.udn
+++ b/Engine/Documentation/Source/Programming/UnrealBuildSystem/TargetFiles/TargetFilesProperties/TargetFilesProperties.INT.udn
@@ -174,6 +174,8 @@ $ bForceNoAutoRTFMCompiler (Boolean): Whether to use force AutoRTFM Clang compil
 
 $ bUseAutoRTFMVerifier (Boolean): Whether to enable emitting AutoRTFM verification metadata
 
+$ bAutoRTFMVerify (Boolean): Whether to run LLVM verification after the AutoRTFM compiler pass. This is used by our compiler folks to ensure the pass works with the various code-paths UBT can take us down.
+
 $ bCompileChaos (Boolean): Whether to compile the Chaos physics plugin.
 
 $ bUseChaos (Boolean): Whether to use the Chaos physics interface. This overrides the physx flags to disable APEX and NvCloth
