Documentation for build tools

diff --git a/Engine/Documentation/Source/Programming/Development/BuildGraph/ScriptAnatomy/Tasks/PredefinedTasks/BuildGraphScriptPredefinedTasks.INT.udn b/Engine/Documentation/Source/Programming/Development/BuildGraph/ScriptAnatomy/Tasks/PredefinedTasks/BuildGraphScriptPredefinedTasks.INT.udn
--- a/Engine/Documentation/Source/Programming/Development/BuildGraph/ScriptAnatomy/Tasks/PredefinedTasks/BuildGraphScriptPredefinedTasks.INT.udn
+++ b/Engine/Documentation/Source/Programming/Development/BuildGraph/ScriptAnatomy/Tasks/PredefinedTasks/BuildGraphScriptPredefinedTasks.INT.udn
@@ -137,6 +137,7 @@ Copies files from one directory to another.
 | Files           | FileSpec  | Optional | Optional filter to be applied to the list of input files.                         |
 | From            | FileSpec  | Required | The pattern(s) to copy from (for example, Engine/*.txt).                          |
 | To              | FileSpec  | Required | The directory to copy to.                                                         |
+| Exclude         | FileSpec  | Optional | Optional pattern(s) to exclude from the copy for example, Engine/NoCopy*.txt)     |
 | Overwrite       | Boolean   | Optional | Whether or not to overwrite existing files.                                       |
 | Tag             | TagList   | Optional | Tag to be applied to build products of this task.                                 |
 | ErrorIfNotFound | Boolean   | Optional | Whether or not to throw an error if no files were found to copy                   |
diff --git a/Engine/Documentation/Source/Programming/UnrealBuildSystem/ModuleFiles/ModuleFilesProperties/ModuleFilesProperties.INT.udn b/Engine/Documentation/Source/Programming/UnrealBuildSystem/ModuleFiles/ModuleFilesProperties/ModuleFilesProperties.INT.udn
--- a/Engine/Documentation/Source/Programming/UnrealBuildSystem/ModuleFiles/ModuleFilesProperties/ModuleFilesProperties.INT.udn
+++ b/Engine/Documentation/Source/Programming/UnrealBuildSystem/ModuleFiles/ModuleFilesProperties/ModuleFilesProperties.INT.udn
@@ -134,6 +134,8 @@ $ bValidateFormatStrings (Boolean): Emits compilation errors for incorrect UE_LO
 
 $ bValidateInternalApi (Boolean): Emits deprecated warnings\errors for internal API usage for non-engine modules
 
+$ bValidateExperimentalApi (Boolean): Emits deprecated warnings\errors for experimental API usage for non-engine modules
+
 $ DefaultBuildSettings (BuildSettingsVersion): Which engine version's build settings to use by default.
 
 $ IncludeOrderVersion (EngineIncludeOrderVersion): What version of include order to use when compiling this module. Can be overridden via -ForceIncludeOrder on the command line or in a module's rules.
