Documentation for build tools

diff --git a/Engine/Documentation/Source/Programming/Development/BuildGraph/ScriptAnatomy/Tasks/PredefinedTasks/BuildGraphScriptPredefinedTasks.INT.udn b/Engine/Documentation/Source/Programming/Development/BuildGraph/ScriptAnatomy/Tasks/PredefinedTasks/BuildGraphScriptPredefinedTasks.INT.udn
--- a/Engine/Documentation/Source/Programming/Development/BuildGraph/ScriptAnatomy/Tasks/PredefinedTasks/BuildGraphScriptPredefinedTasks.INT.udn
+++ b/Engine/Documentation/Source/Programming/Development/BuildGraph/ScriptAnatomy/Tasks/PredefinedTasks/BuildGraphScriptPredefinedTasks.INT.udn
@@ -647,6 +647,17 @@ Creates a new changelist and submits a set of files to a Perforce stream.
 | P4Verbose       | Boolean            | Optional | Allow verbose P4 output (spew).                                                                                                                                              |
 | If              | Condition          | Optional | Whether to execute this task. It is ignored if this condition evaluates to false.                                                                                            |
 
+### SymGen
+
+Generates a portable symbol dump file from the specified binaries
+
+|           |           |          |                                                                                                            |
+| --------- | --------- | -------- | ---------------------------------------------------------------------------------------------------------- |
+| Files     | FileSpec  | Required | List of file specifications separated by semicolons (eg. *.cpp;Engine/.../*.bat), or the name of a tag set |
+| Tag       | TagList   | Optional | Tag to be applied to build products of this task                                                           |
+| UseRadSym | Boolean   | Optional | If set, this will use the rad debugger pdb symbol dumper as well as the rad symbol_path_fixer.             |
+| If        | Condition | Optional | Whether to execute this task. It is ignored if this condition evaluates to false.                          |
+
 ### SymStore
 
 Task that strips symbols from a set of files.
