Documentation for build tools

diff --git a/Engine/Documentation/Source/Programming/UnrealBuildSystem/Configuration/BuildConfigProperties/BuildConfigProperties.INT.udn b/Engine/Documentation/Source/Programming/UnrealBuildSystem/Configuration/BuildConfigProperties/BuildConfigProperties.INT.udn
--- a/Engine/Documentation/Source/Programming/UnrealBuildSystem/Configuration/BuildConfigProperties/BuildConfigProperties.INT.udn
+++ b/Engine/Documentation/Source/Programming/UnrealBuildSystem/Configuration/BuildConfigProperties/BuildConfigProperties.INT.udn
@@ -322,6 +322,8 @@ $ bClangStandaloneDebug : True if optimizations to reduce the size of debug info
 
 $ bAllowClangLinker : True if we should use the Clang linker (LLD) when we are compiling with Clang or Intel oneAPI, otherwise we use the MSVC linker.
 
+$ bAllowRadLinker : True if we should use the Rad linker
+
 $ WindowsSdkVersion : The specific Windows SDK version to use. This may be a specific version number (for example, "8.1", "10.0" or "10.0.10150.0"), or the string "Latest", to select the newest available version. By default, and if it is available, we use the Windows SDK version indicated by WindowsPlatform.DefaultWindowsSdkVersion (otherwise, we use the latest version).
 
 $ bWriteSarif : Whether .sarif files containing errors and warnings are written alongside each .obj, if supported
