Documentation for build tools

diff --git a/Engine/Documentation/Source/Programming/UnrealBuildSystem/ModuleFiles/ModuleFilesProperties/ModuleFilesProperties.INT.udn b/Engine/Documentation/Source/Programming/UnrealBuildSystem/ModuleFiles/ModuleFilesProperties/ModuleFilesProperties.INT.udn
--- a/Engine/Documentation/Source/Programming/UnrealBuildSystem/ModuleFiles/ModuleFilesProperties/ModuleFilesProperties.INT.udn
+++ b/Engine/Documentation/Source/Programming/UnrealBuildSystem/ModuleFiles/ModuleFilesProperties/ModuleFilesProperties.INT.udn
@@ -160,6 +160,8 @@ $ bEnableObjCAutomaticReferenceCounting (Boolean): Enable objective C automatic
 
 $ DeterministicWarningLevel (WarningLevel): How to treat deterministic warnings (experimental).
 
+$ CppCompileWarningSettings (CppCompileWarnings): Settings object used by the module.
+
 $ ShadowVariableWarningLevel (WarningLevel): How to treat shadow variable warnings
 
 $ bWarningsAsErrors (Boolean): Whether to enable all warnings as errors. UE enables most warnings as errors already, but disables a few (such as deprecation warnings).
diff --git a/Engine/Documentation/Source/Programming/UnrealBuildSystem/TargetFiles/TargetFilesProperties/TargetFilesProperties.INT.udn b/Engine/Documentation/Source/Programming/UnrealBuildSystem/TargetFiles/TargetFilesProperties/TargetFilesProperties.INT.udn
--- a/Engine/Documentation/Source/Programming/UnrealBuildSystem/TargetFiles/TargetFilesProperties/TargetFilesProperties.INT.udn
+++ b/Engine/Documentation/Source/Programming/UnrealBuildSystem/TargetFiles/TargetFilesProperties/TargetFilesProperties.INT.udn
@@ -390,6 +390,8 @@ $ bValidateFormatStrings (Boolean): Emits compilation errors for incorrect UE_LO
 
 $ DeprecationWarningLevel (WarningLevel): Level to report deprecation warnings as errors
 
+$ CppCompileWarningSettings (CppCompileWarnings): C++ Warnings settings object used by the target
+
 $ ShadowVariableWarningLevel (WarningLevel): Forces shadow variable warnings to be treated as errors on platforms that support it.
 
 $ bWarningsAsErrors (Boolean): Whether to enable all warnings as errors. UE enables most warnings as errors already, but disables a few (such as deprecation warnings).
