Documentation for build tools

diff --git a/Engine/Documentation/Source/Programming/UnrealBuildSystem/Configuration/BuildConfigProperties/BuildConfigProperties.INT.udn b/Engine/Documentation/Source/Programming/UnrealBuildSystem/Configuration/BuildConfigProperties/BuildConfigProperties.INT.udn
--- a/Engine/Documentation/Source/Programming/UnrealBuildSystem/Configuration/BuildConfigProperties/BuildConfigProperties.INT.udn
+++ b/Engine/Documentation/Source/Programming/UnrealBuildSystem/Configuration/BuildConfigProperties/BuildConfigProperties.INT.udn
@@ -56,6 +56,8 @@ $ UnsafeTypeCastWarningLevel : Indicates what warning/error level to treat unsaf
 
 $ UndefinedIdentifierWarningLevel : Indicates what warning/error level to treat undefined identifiers in conditional expressions.
 
+$ SwitchUnhandledEnumeratorWarningLevel : Indicates what warning/error level to treat unhandled enumerators in switches on enumeration-typed values.
+
 $ DefaultWarningLevel : Default treatment of uncategorized warnings
 
 $ DeprecationWarningLevel : Level to report deprecation warnings as errors
