Documentation for build tools

diff --git a/Engine/Documentation/Source/Programming/Development/BuildGraph/ScriptAnatomy/Tasks/PredefinedTasks/BuildGraphScriptPredefinedTasks.INT.udn b/Engine/Documentation/Source/Programming/Development/BuildGraph/ScriptAnatomy/Tasks/PredefinedTasks/BuildGraphScriptPredefinedTasks.INT.udn
--- a/Engine/Documentation/Source/Programming/Development/BuildGraph/ScriptAnatomy/Tasks/PredefinedTasks/BuildGraphScriptPredefinedTasks.INT.udn
+++ b/Engine/Documentation/Source/Programming/Development/BuildGraph/ScriptAnatomy/Tasks/PredefinedTasks/BuildGraphScriptPredefinedTasks.INT.udn
@@ -546,6 +546,21 @@ Renames a file, or group of files.
 | Tag   | TagList   | Optional | Tag to be applied to the renamed files.                                                                  |
 | If    | Condition | Optional | Whether to execute this task. It is ignored if this condition evaluates to false.                        |
 
+### RetrieveArtifact
+
+Retrieves an artifact from Horde
+
+|           |           |          |                                                                                       |
+| --------- | --------- | -------- | ------------------------------------------------------------------------------------- |
+| Stream    | String    | Optional | Stream containing the artifact                                                        |
+| Commit    | String    | Optional | Change number for the artifact                                                        |
+| MaxCommit | String    | Optional | Maximum commit for the artifact                                                       |
+| Name      | String    | Optional | Name of the artifact                                                                  |
+| Type      | String    | Required | The artifact type. Determines the permissions and expiration policy for the artifact. |
+| Keys      | String    | Optional | Keys for the artifact                                                                 |
+| OutputDir | String    | Required | Output directory for                                                                  |
+| If        | Condition | Optional | Whether to execute this task. It is ignored if this condition evaluates to false.     |
+
 ### SanitizeReceipt
 
 Task that tags build products and/or runtime dependencies by reading from *.target files.
