Documentation for build tools

diff --git a/Engine/Documentation/Source/Programming/UnrealBuildSystem/Configuration/BuildConfigProperties/BuildConfigProperties.INT.udn b/Engine/Documentation/Source/Programming/UnrealBuildSystem/Configuration/BuildConfigProperties/BuildConfigProperties.INT.udn
--- a/Engine/Documentation/Source/Programming/UnrealBuildSystem/Configuration/BuildConfigProperties/BuildConfigProperties.INT.udn
+++ b/Engine/Documentation/Source/Programming/UnrealBuildSystem/Configuration/BuildConfigProperties/BuildConfigProperties.INT.udn
@@ -142,6 +142,8 @@ $ bUsePCHFiles : Whether PCH files should be used.
 
 $ bDeterministic : Set flags require for deterministic compiling and linking. Enabling deterministic mode for msvc disables codegen multithreading so compiling will be slower
 
+$ bCompilerVFS : Experimental work in progress: Set flags to use toolchain virtual file system support, to generate consistent pathing in outputs
+
 $ bChainPCHs : Whether PCHs should be chained when compiling with clang.
 
 $ bForceIncludePCHHeadersForGenCppFilesWhenPCHIsDisabled : Whether PCH headers should be force included for gen.cpp files when PCH is disabled.
diff --git a/Engine/Documentation/Source/Programming/UnrealBuildSystem/ModuleFiles/ModuleFilesProperties/ModuleFilesProperties.INT.udn b/Engine/Documentation/Source/Programming/UnrealBuildSystem/ModuleFiles/ModuleFilesProperties/ModuleFilesProperties.INT.udn
--- a/Engine/Documentation/Source/Programming/UnrealBuildSystem/ModuleFiles/ModuleFilesProperties/ModuleFilesProperties.INT.udn
+++ b/Engine/Documentation/Source/Programming/UnrealBuildSystem/ModuleFiles/ModuleFilesProperties/ModuleFilesProperties.INT.udn
@@ -106,6 +106,8 @@ $ SubclassRules (List&lt;String&gt;): Subclass rules files which invalidate the
 
 $ GenerateHeaderFuncs (List&lt;ValueTuple&lt;String, Action&lt;ILogger, DirectoryReference&gt;&gt;&gt;): A list of subdirectory names and functions that are invoked to generate header files. The subdirectory name is appended to the generated code directory to form a new directory that headers are generated inside.
 
+$ SwiftInteropHeader (String): If set, this module will generate a bridging header between C++ and Swift, using this input header
+
 $ Name (String): Name of this module
 
 $ Type (ModuleType): Type of module
diff --git a/Engine/Documentation/Source/Programming/UnrealBuildSystem/TargetFiles/TargetFilesProperties/TargetFilesProperties.INT.udn b/Engine/Documentation/Source/Programming/UnrealBuildSystem/TargetFiles/TargetFilesProperties/TargetFilesProperties.INT.udn
--- a/Engine/Documentation/Source/Programming/UnrealBuildSystem/TargetFiles/TargetFilesProperties/TargetFilesProperties.INT.udn
+++ b/Engine/Documentation/Source/Programming/UnrealBuildSystem/TargetFiles/TargetFilesProperties/TargetFilesProperties.INT.udn
@@ -442,6 +442,8 @@ $ bUsePCHFiles (Boolean): Whether PCH files should be used.
 
 $ bDeterministic (Boolean): Set flags require for deterministic compiling and linking. Enabling deterministic mode for msvc disables codegen multithreading so compiling will be slower
 
+$ bCompilerVFS (Boolean): Experimental work in progress: Set flags to use toolchain virtual file system support, to generate consistent pathing in outputs
+
 $ bChainPCHs (Boolean): Whether PCHs should be chained when compiling with clang.
 
 $ bForceIncludePCHHeadersForGenCppFilesWhenPCHIsDisabled (Boolean): Whether PCH headers should be force included for gen.cpp files when PCH is disabled.
