Documentation for build tools

diff --git a/Engine/Documentation/Source/Programming/UnrealBuildSystem/Configuration/BuildConfigProperties/BuildConfigProperties.INT.udn b/Engine/Documentation/Source/Programming/UnrealBuildSystem/Configuration/BuildConfigProperties/BuildConfigProperties.INT.udn
--- a/Engine/Documentation/Source/Programming/UnrealBuildSystem/Configuration/BuildConfigProperties/BuildConfigProperties.INT.udn
+++ b/Engine/Documentation/Source/Programming/UnrealBuildSystem/Configuration/BuildConfigProperties/BuildConfigProperties.INT.udn
@@ -84,6 +84,8 @@ $ bUseAutoRTFMVerifier : Whether to enable emitting AutoRTFM verification metada
 
 $ bAutoRTFMVerify : Whether to run LLVM verification after the AutoRTFM compiler pass. This is used by our compiler folks to ensure the pass works with the various code-paths UBT can take us down.
 
+$ bAutoRTFMClosedStaticLinkage : Whether to link closed function declarations statically.
+
 $ bUseInlining : Enable inlining for all modules.
 
 $ bUseDebugLiveCodingConsole : Whether to enable support for live coding
diff --git a/Engine/Documentation/Source/Programming/UnrealBuildSystem/TargetFiles/TargetFilesProperties/TargetFilesProperties.INT.udn b/Engine/Documentation/Source/Programming/UnrealBuildSystem/TargetFiles/TargetFilesProperties/TargetFilesProperties.INT.udn
--- a/Engine/Documentation/Source/Programming/UnrealBuildSystem/TargetFiles/TargetFilesProperties/TargetFilesProperties.INT.udn
+++ b/Engine/Documentation/Source/Programming/UnrealBuildSystem/TargetFiles/TargetFilesProperties/TargetFilesProperties.INT.udn
@@ -176,6 +176,8 @@ $ bUseAutoRTFMVerifier (Boolean): Whether to enable emitting AutoRTFM verificati
 
 $ bAutoRTFMVerify (Boolean): Whether to run LLVM verification after the AutoRTFM compiler pass. This is used by our compiler folks to ensure the pass works with the various code-paths UBT can take us down.
 
+$ bAutoRTFMClosedStaticLinkage (Boolean): Whether to link closed function declarations statically.
+
 $ bCompileChaos (Boolean): Whether to compile the Chaos physics plugin.
 
 $ bUseChaos (Boolean): Whether to use the Chaos physics interface. This overrides the physx flags to disable APEX and NvCloth
