Documentation for build tools

diff --git a/Engine/Documentation/Source/Programming/UnrealBuildSystem/Configuration/BuildConfigProperties/BuildConfigProperties.INT.udn b/Engine/Documentation/Source/Programming/UnrealBuildSystem/Configuration/BuildConfigProperties/BuildConfigProperties.INT.udn
--- a/Engine/Documentation/Source/Programming/UnrealBuildSystem/Configuration/BuildConfigProperties/BuildConfigProperties.INT.udn
+++ b/Engine/Documentation/Source/Programming/UnrealBuildSystem/Configuration/BuildConfigProperties/BuildConfigProperties.INT.udn
@@ -60,6 +60,10 @@ $ SwitchUnhandledEnumeratorWarningLevel : Indicates what warning/error level to
 
 $ PCHPerformanceIssueWarningLevel : Indicates what warning/error level to treat potential PCH performance issues.
 
+$ ModuleUnsupportedWarningLevel : How to treat module unsupported validation messages
+
+$ PluginModuleUnsupportedWarningLevel : How to treat plugin specific module unsupported validation messages
+
 $ ModuleIncludePathWarningLevel : How to treat general module include path validation messages
 
 $ ModuleIncludePrivateWarningLevel : How to treat private module include path validation messages, where a module is adding an include path that exposes private headers
diff --git a/Engine/Documentation/Source/Programming/UnrealBuildSystem/TargetFiles/TargetFilesProperties/TargetFilesProperties.INT.udn b/Engine/Documentation/Source/Programming/UnrealBuildSystem/TargetFiles/TargetFilesProperties/TargetFilesProperties.INT.udn
--- a/Engine/Documentation/Source/Programming/UnrealBuildSystem/TargetFiles/TargetFilesProperties/TargetFilesProperties.INT.udn
+++ b/Engine/Documentation/Source/Programming/UnrealBuildSystem/TargetFiles/TargetFilesProperties/TargetFilesProperties.INT.udn
@@ -604,4 +604,6 @@ $ MemoryPerActionGB (Double): Max amount of memory that each compile action may
 
 $ GeneratedProjectName (String): When generating project files, specifies the name of the project file to use when there are multiple targets of the same type.
 
+$ StaticAllocator (StaticAllocatorType): Statically selected allocator to use
+
 
