Chaos::Add detailed descriptions of some joint parameters.

diff --git a/Engine/Documentation/Source/Shared/Types/UPhysicsConstraintComponent/UPhysicsConstraintComponent.INT.udn b/Engine/Documentation/Source/Shared/Types/UPhysicsConstraintComponent/UPhysicsConstraintComponent.INT.udn
--- a/Engine/Documentation/Source/Shared/Types/UPhysicsConstraintComponent/UPhysicsConstraintComponent.INT.udn
+++ b/Engine/Documentation/Source/Shared/Types/UPhysicsConstraintComponent/UPhysicsConstraintComponent.INT.udn
@@ -5,17 +5,17 @@ Crumbs:
 
 
 [EXCERPT:ConstraintActor1]
-
+This is the first actor attached to the joint. The joint is not effective if this is empty. 
 [/EXCERPT:ConstraintActor1]
 
 [EXCERPT:ComponentName1]
-
+If you would like to constrain a specific component of the specified actor above, you can provide a Component Name here. If the root of the Actor is of a type that can be constrained, it will be the default Component that is constrained. However, if you provide a valid Component Name here, that Component will become the target of the Physics Constraint. If that Component is a Skeletal Mesh, you must also provide a Bone name in the appropriate property.
 [/EXCERPT:ComponentName1]
 
 [EXCERPT:ConstraintActor2]
-
+This is the second actor attached to the joint. 
 [/EXCERPT:ConstraintActor2]
 
 [EXCERPT:ComponentName2]
-
+If you would like to constrain a specific component of the specified actor above, you can provide a Component Name here. If the root of the Actor is of a type that can be constrained, it will be the default Component that is constrained. However, if you provide a valid Component Name here, that Component will become the target of the Physics Constraint. If that Component is a Skeletal Mesh, you must also provide a Bone name in the appropriate property.
 [/EXCERPT:ComponentName2]
diff --git a/Engine/Source/Runtime/Engine/Classes/PhysicsEngine/ConstraintInstance.h b/Engine/Source/Runtime/Engine/Classes/PhysicsEngine/ConstraintInstance.h
--- a/Engine/Source/Runtime/Engine/Classes/PhysicsEngine/ConstraintInstance.h
+++ b/Engine/Source/Runtime/Engine/Classes/PhysicsEngine/ConstraintInstance.h
@@ -148,7 +148,7 @@ struct FConstraintProfileProperties
 	UPROPERTY(EditAnywhere, Category = Projection)
 	uint8 bEnableMassConditioning : 1;
 
-	/** Whether to use linear solver for this joint. Linear solver is faster and less accurate. The nonlinear solver is slower and more convergent.  */
+	/** World Solver: Whether to use linear solver for this joint. Linear solver is faster and less accurate. The nonlinear solver is slower and more convergent.  */
 	UPROPERTY(EditAnywhere, Category = Projection)
 	uint8 bUseLinearJointSolver : 1;
 
