Documentation for build tools

diff --git a/Engine/Documentation/Source/Programming/UnrealBuildSystem/Configuration/BuildConfigProperties/BuildConfigProperties.INT.udn b/Engine/Documentation/Source/Programming/UnrealBuildSystem/Configuration/BuildConfigProperties/BuildConfigProperties.INT.udn
--- a/Engine/Documentation/Source/Programming/UnrealBuildSystem/Configuration/BuildConfigProperties/BuildConfigProperties.INT.udn
+++ b/Engine/Documentation/Source/Programming/UnrealBuildSystem/Configuration/BuildConfigProperties/BuildConfigProperties.INT.udn
@@ -272,6 +272,8 @@ $ bTuneDebugInfoForLLDB : Turns on tuning of debug info for LLDB
 
 $ bDisableDumpSyms : Whether to globally disable calling dump_syms
 
+$ bEnableInstrumentation : Enables instrumentation.
+
 $ bWriteSolutionOptionFile : Whether to write a solution option (suo) file for the sln.
 
 $ bVsConfigFile : Whether to write a .vsconfig file next to the sln to suggest components to install.
diff --git a/Engine/Documentation/Source/Programming/UnrealBuildSystem/ModuleFiles/ModuleFilesProperties/ModuleFilesProperties.INT.udn b/Engine/Documentation/Source/Programming/UnrealBuildSystem/ModuleFiles/ModuleFilesProperties/ModuleFilesProperties.INT.udn
--- a/Engine/Documentation/Source/Programming/UnrealBuildSystem/ModuleFiles/ModuleFilesProperties/ModuleFilesProperties.INT.udn
+++ b/Engine/Documentation/Source/Programming/UnrealBuildSystem/ModuleFiles/ModuleFilesProperties/ModuleFilesProperties.INT.udn
@@ -222,6 +222,8 @@ $ PrecompileForTargets (PrecompileTargetsType): Which targets this module should
 
 $ bRequiresImplementModule (Nullable&lt;Boolean&gt;): Whether this module requires the IMPLEMENT_MODULE macro to be implemented. Most UE modules require this, since we use the IMPLEMENT_MODULE macro to do other global overloads (eg. operator new/delete forwarding to GMalloc).
 
+$ ExtraRootPath (String): Module can provide ExtraRootPath. This will be part of the rootpaths table and is needed if module is using paths outside normal environment
+
 $ VersePath (String): If this module has associated Verse code, this is the Verse root path of it
 
 $ VerseScope (VerseScope): Visibility of Verse code in this module's Source/Verse folder
