Documentation for build tools

diff --git a/Engine/Documentation/Source/Programming/UnrealBuildSystem/ModuleFiles/ModuleFilesProperties/ModuleFilesProperties.INT.udn b/Engine/Documentation/Source/Programming/UnrealBuildSystem/ModuleFiles/ModuleFilesProperties/ModuleFilesProperties.INT.udn
--- a/Engine/Documentation/Source/Programming/UnrealBuildSystem/ModuleFiles/ModuleFilesProperties/ModuleFilesProperties.INT.udn
+++ b/Engine/Documentation/Source/Programming/UnrealBuildSystem/ModuleFiles/ModuleFilesProperties/ModuleFilesProperties.INT.udn
@@ -188,6 +188,8 @@ $ bUseUnity (Boolean): If unity builds are enabled this can be used to override
 
 $ bMergeUnityFiles (Boolean): Whether to merge module and generated unity files for faster compilation.
 
+$ UnityFileTypes (UnityFileType): Source file types that will be placed into unity files when enabled.
+
 $ MinSourceFilesForUnityBuildOverride (Int32): The number of source files in this module before unity build will be activated for that module.  If set to anything besides -1, will override the default setting which is controlled by MinGameModuleSourceFilesForUnityBuild
 
 $ MinFilesUsingPrecompiledHeaderOverride (Int32): Overrides BuildConfiguration.MinFilesUsingPrecompiledHeader if non-zero.
