Documentation for build tools

diff --git a/Engine/Documentation/Source/Programming/UnrealBuildSystem/Configuration/BuildConfigProperties/BuildConfigProperties.INT.udn b/Engine/Documentation/Source/Programming/UnrealBuildSystem/Configuration/BuildConfigProperties/BuildConfigProperties.INT.udn
--- a/Engine/Documentation/Source/Programming/UnrealBuildSystem/Configuration/BuildConfigProperties/BuildConfigProperties.INT.udn
+++ b/Engine/Documentation/Source/Programming/UnrealBuildSystem/Configuration/BuildConfigProperties/BuildConfigProperties.INT.udn
@@ -444,8 +444,6 @@ $ bDisableWaitOnMem : Disable UBA waiting on available memory before spawning ne
 
 $ bAllowKillOnMem : Let UBA kill running processes when close to out of memory
 
-$ OutputStatsThresholdMs : Threshold for when executor should output logging for the process. Defaults to never
-
 $ bWriteToDisk : Skip writing intermediate and output files to disk. Useful for validation builds where we don't need the output
 
 $ bDisableCustomAlloc : Set to true to disable mimalloc and detouring of memory allocations.
diff --git a/Engine/Documentation/Source/Programming/UnrealBuildSystem/ModuleFiles/ModuleFilesProperties/ModuleFilesProperties.INT.udn b/Engine/Documentation/Source/Programming/UnrealBuildSystem/ModuleFiles/ModuleFilesProperties/ModuleFilesProperties.INT.udn
--- a/Engine/Documentation/Source/Programming/UnrealBuildSystem/ModuleFiles/ModuleFilesProperties/ModuleFilesProperties.INT.udn
+++ b/Engine/Documentation/Source/Programming/UnrealBuildSystem/ModuleFiles/ModuleFilesProperties/ModuleFilesProperties.INT.udn
@@ -224,7 +224,7 @@ $ PrecompileForTargets (PrecompileTargetsType): Which targets this module should
 
 $ bRequiresImplementModule (Nullable&lt;Boolean&gt;): Whether this module requires the IMPLEMENT_MODULE macro to be implemented. Most UE modules require this, since we use the IMPLEMENT_MODULE macro to do other global overloads (eg. operator new/delete forwarding to GMalloc).
 
-$ ExtraRootPath (String): Module can provide ExtraRootPath. This will be part of the rootpaths table and is needed if module is using paths outside normal environment
+$ ExtraRootPath (Nullable&lt;ValueTuple&lt;String, String&gt;&gt;): Module can provide ExtraRootPath. This will be part of the rootpaths table and is needed if module is using paths outside normal environment Id should be the same for all extra roots with the same path otherwise we can end up with cache misses that could be cache hits Make sure path is as long as possible. I.e, don't use "Program Files" as root path but rather "Program Files/MyProgram"
 
 $ VersePath (String): If this module has associated Verse code, this is the Verse root path of it
 
