Documentation for build tools

diff --git a/Engine/Documentation/Source/Programming/UnrealBuildSystem/Configuration/BuildConfigProperties/BuildConfigProperties.INT.udn b/Engine/Documentation/Source/Programming/UnrealBuildSystem/Configuration/BuildConfigProperties/BuildConfigProperties.INT.udn
--- a/Engine/Documentation/Source/Programming/UnrealBuildSystem/Configuration/BuildConfigProperties/BuildConfigProperties.INT.udn
+++ b/Engine/Documentation/Source/Programming/UnrealBuildSystem/Configuration/BuildConfigProperties/BuildConfigProperties.INT.udn
@@ -148,7 +148,7 @@ $ bUsePCHFiles : Whether PCH files should be used.
 
 $ bDeterministic : Set flags require for deterministic compiling and linking. Enabling deterministic mode for msvc disables codegen multithreading so compiling will be slower
 
-$ bCompilerVFS : Experimental work in progress: Set flags to use toolchain virtual file system support, to generate consistent pathing in outputs
+$ bUseVFS : Experimental work in progress: Set flags to use toolchain virtual file system support, to generate consistent pathing in outputs
 
 $ bChainPCHs : Whether PCHs should be chained when compiling with clang.
 
diff --git a/Engine/Documentation/Source/Programming/UnrealBuildSystem/TargetFiles/TargetFilesProperties/TargetFilesProperties.INT.udn b/Engine/Documentation/Source/Programming/UnrealBuildSystem/TargetFiles/TargetFilesProperties/TargetFilesProperties.INT.udn
--- a/Engine/Documentation/Source/Programming/UnrealBuildSystem/TargetFiles/TargetFilesProperties/TargetFilesProperties.INT.udn
+++ b/Engine/Documentation/Source/Programming/UnrealBuildSystem/TargetFiles/TargetFilesProperties/TargetFilesProperties.INT.udn
@@ -444,7 +444,7 @@ $ bUsePCHFiles (Boolean): Whether PCH files should be used.
 
 $ bDeterministic (Boolean): Set flags require for deterministic compiling and linking. Enabling deterministic mode for msvc disables codegen multithreading so compiling will be slower
 
-$ bCompilerVFS (Boolean): Experimental work in progress: Set flags to use toolchain virtual file system support, to generate consistent pathing in outputs
+$ bUseVFS (Boolean): Experimental work in progress: Set flags to use toolchain virtual file system support, to generate consistent pathing in outputs
 
 $ bChainPCHs (Boolean): Whether PCHs should be chained when compiling with clang.
 
