Documentation for build tools

diff --git a/Engine/Documentation/Source/Programming/Development/BuildGraph/ScriptAnatomy/Tasks/PredefinedTasks/BuildGraphScriptPredefinedTasks.INT.udn b/Engine/Documentation/Source/Programming/Development/BuildGraph/ScriptAnatomy/Tasks/PredefinedTasks/BuildGraphScriptPredefinedTasks.INT.udn
--- a/Engine/Documentation/Source/Programming/Development/BuildGraph/ScriptAnatomy/Tasks/PredefinedTasks/BuildGraphScriptPredefinedTasks.INT.udn
+++ b/Engine/Documentation/Source/Programming/Development/BuildGraph/ScriptAnatomy/Tasks/PredefinedTasks/BuildGraphScriptPredefinedTasks.INT.udn
@@ -833,6 +833,20 @@ Imports an oplog from Zen to a specified destination.
 | Key         | String        | Optional | What key to use when importing from cloud                                                          |
 | If          | Condition     | Optional | Whether to execute this task. It is ignored if this condition evaluates to false.                  |
 
+### ZenImportSnapshot
+
+Imports a snapshot from a specified destination to Zen.
+
+|                           |               |          |                                                                                                                                    |
+| ------------------------- | ------------- | -------- | ---------------------------------------------------------------------------------------------------------------------------------- |
+| Project                   | FileReference | Required | The project into which snapshots should be imported                                                                                |
+| SnapshotDescriptorFile    | FileReference | Optional | A file to read with information about the snapshots to be imported                                                                 |
+| SnapshotDescriptorJSON    | String        | Optional | A JSON blob with information about the snapshots to be imported                                                                    |
+| OverridePlatformCookedDir | String        | Optional | Optional. Where to look for the ue.projectstore file. The pattern {Platform} can be used for importing multiple platforms at once. |
+| ForceImport               | Boolean       | Optional | Optional. If true, force import of the oplog (corresponds to --force on the Zen command line).                                     |
+| AsyncImport               | Boolean       | Optional | Optional. If true, import oplogs asynchronously (corresponds to --async on the Zen command line).                                  |
+| If                        | Condition     | Optional | Whether to execute this task. It is ignored if this condition evaluates to false.                                                  |
+
 ### ZenLaunch
 
 Launches ZenServer
