Documentation for build tools

diff --git a/Engine/Documentation/Source/Programming/Development/BuildGraph/ScriptAnatomy/Tasks/PredefinedTasks/BuildGraphScriptPredefinedTasks.INT.udn b/Engine/Documentation/Source/Programming/Development/BuildGraph/ScriptAnatomy/Tasks/PredefinedTasks/BuildGraphScriptPredefinedTasks.INT.udn
--- a/Engine/Documentation/Source/Programming/Development/BuildGraph/ScriptAnatomy/Tasks/PredefinedTasks/BuildGraphScriptPredefinedTasks.INT.udn
+++ b/Engine/Documentation/Source/Programming/Development/BuildGraph/ScriptAnatomy/Tasks/PredefinedTasks/BuildGraphScriptPredefinedTasks.INT.udn
@@ -793,6 +793,7 @@ Exports an snapshot from Zen to a specified destination.
 | ---------------------------------- | ------------------ | -------- | ---------------------------------------------------------------------------------------------------------------------------- |
 | Project                            | FileReference      | Optional | The project from which to export the snapshot                                                                                |
 | Platform                           | String             | Optional | The target platform(s) to export the snapshot for                                                                            |
+| Metadata                           | String             | Optional | The metadata to associate with the snapshot                                                                                  |
 | SnapshotBaseDescriptorFile         | FileReference      | Optional | A file to read with information about the snapshot that should be used as a base when exporting this new snapshot            |
 | SnapshotDescriptorFile             | FileReference      | Optional | A file to create with information about the snapshot that was exported                                                       |
 | DestinationStorageType             | String             | Required | The type of destination to export the snapshot to (cloud, ...)                                                               |
