Documentation for build tools

diff --git a/Engine/Documentation/Source/Programming/Development/BuildGraph/ScriptAnatomy/Tasks/PredefinedTasks/BuildGraphScriptPredefinedTasks.INT.udn b/Engine/Documentation/Source/Programming/Development/BuildGraph/ScriptAnatomy/Tasks/PredefinedTasks/BuildGraphScriptPredefinedTasks.INT.udn
--- a/Engine/Documentation/Source/Programming/Development/BuildGraph/ScriptAnatomy/Tasks/PredefinedTasks/BuildGraphScriptPredefinedTasks.INT.udn
+++ b/Engine/Documentation/Source/Programming/Development/BuildGraph/ScriptAnatomy/Tasks/PredefinedTasks/BuildGraphScriptPredefinedTasks.INT.udn
@@ -733,25 +733,6 @@ Extract files from a zip archive.
 | Tag            | TagList            | Optional | Tag to be applied to the extracted files.                                         |
 | If             | Condition          | Optional | Whether to execute this task. It is ignored if this condition evaluates to false. |
 
-### Upload
-
-Uploads a set of files to Jupiter for future retrival
-
-|                    |                    |          |                                                                                                                                                                                       |
-| ------------------ | ------------------ | -------- | ------------------------------------------------------------------------------------------------------------------------------------------------------------------------------------- |
-| FromDir            | DirectoryReference | Required | The directory to read compressed files from.                                                                                                                                          |
-| Files              | FileSpec           | Optional | List of file specifications separated by semicolons (for example, *.cpp;Engine/.../*.bat), or the name of a tag set. Relative paths are taken from FromDir.                           |
-| JupiterNamespace   | String             | Required | The jupiter namespace used to upload the build. Used to control who has access to the build.                                                                                          |
-| JupiterKey         | String             | Required | The key of the build as will be used to download the build again. This has to be globally unique for this particular upload.                                                          |
-| ArchiveType        | String             | Required | The type of archive these files are from, will be added to the metadata                                                                                                               |
-| ProjectName        | String             | Required | The name of the project this set of files are associated with, will be added to the metadata                                                                                          |
-| Branch             | String             | Required | The source control branch these files were generated from, will be added to the metadata                                                                                              |
-| Changelist         | String             | Required | The source control revision these files were generated from, will be added to the metadata                                                                                            |
-| JupiterUrl         | String             | Required | Specify the url to the Jupiter instance to upload to                                                                                                                                  |
-| AdditionalMetadata | String             | Optional | Semi-colon separated list of '=' separated key value mappings to add to the metadata. E.g. Foo=bar;spam=eggs                                                                          |
-| LimitMemoryUsage   | Boolean            | Optional | If enabled file content is not kept in memory, results in lower memory usage but increased io as file contents needs to be read multiple times (for hashing as well as during upload) |
-| If                 | Condition          | Optional | Whether to execute this task. It is ignored if this condition evaluates to false.                                                                                                     |
-
 ### Wait
 
 Waits a defined number of seconds.
