Documentation for build tools

diff --git a/Engine/Documentation/Source/Programming/Development/BuildGraph/ScriptAnatomy/Tasks/PredefinedTasks/BuildGraphScriptPredefinedTasks.INT.udn b/Engine/Documentation/Source/Programming/Development/BuildGraph/ScriptAnatomy/Tasks/PredefinedTasks/BuildGraphScriptPredefinedTasks.INT.udn
--- a/Engine/Documentation/Source/Programming/Development/BuildGraph/ScriptAnatomy/Tasks/PredefinedTasks/BuildGraphScriptPredefinedTasks.INT.udn
+++ b/Engine/Documentation/Source/Programming/Development/BuildGraph/ScriptAnatomy/Tasks/PredefinedTasks/BuildGraphScriptPredefinedTasks.INT.udn
@@ -162,6 +162,30 @@ Uploads an artifact to Horde
 | Dedupe      | Boolean   | Optional | Whether to reuse blobs from the previous artifact                                                                              |
 | If          | Condition | Optional | Whether to execute this task. It is ignored if this condition evaluates to false.                                              |
 
+### CreateCloudArtifact
+
+Uploads an artifact to Cloud DDC
+
+|                          |           |          |                                                                                                                                |
+| ------------------------ | --------- | -------- | ------------------------------------------------------------------------------------------------------------------------------ |
+| Name                     | String    | Required | Name of the artifact                                                                                                           |
+| Type                     | String    | Required | The artifact type. Determines the permissions and expiration policy for the artifact.                                          |
+| Description              | String    | Optional | Description for the artifact. Will be shown through the Horde dashboard.                                                       |
+| Platform                 | String    | Required | The platform the artifact is created for                                                                                       |
+| Project                  | String    | Required | The project this artifact is created for                                                                                       |
+| Host                     | String    | Optional | The platform the artifact is created for                                                                                       |
+| Namespace                | String    | Optional | The platform the artifact is created for                                                                                       |
+| AccessToken              | String    | Optional | The access token to use                                                                                                        |
+| BaseDir                  | String    | Optional | Base path for the uploaded files. All the tagged files must be under this directory. Defaults to the workspace root directory. |
+| StreamId                 | String    | Optional | Stream containing the artifact.                                                                                                |
+| Commit                   | String    | Optional | Commit for the uploaded artifact.                                                                                              |
+| Files                    | FileSpec  | Required | Files to include in the artifact.                                                                                              |
+| Metadata                 | String    | Optional | Other metadata for the artifact, separated by semicolons.                                                                      |
+| Force                    | Boolean   | Optional | Force the blobs to be uploaded no matter if they exist or not                                                                  |
+| BaselineBranch           | String    | Required | The streamid used for your mainline, used to build content as patches on top of this branch if no other baselines exist.       |
+| MaxBlobUploadParallelism | Integer   | Optional | Set to upload more blobs in parallel, defaults to number of cores. Set to -1 to go as wide as possible.                        |
+| If                       | Condition | Optional | Whether to execute this task. It is ignored if this condition evaluates to false.                                              |
+
 ### CsCompile
 
 Compiles C# project files, and their dependencies.
@@ -561,6 +585,27 @@ Retrieves an artifact from Horde
 | OutputDir | String    | Required | Output directory for                                                                  |
 | If        | Condition | Optional | Whether to execute this task. It is ignored if this condition evaluates to false.     |
 
+### RetrieveCloudArtifact
+
+Retrieves an artifact from Cloud DDC
+
+|                      |           |          |                                                                                       |
+| -------------------- | --------- | -------- | ------------------------------------------------------------------------------------- |
+| StreamId             | String    | Optional | Stream containing the artifact                                                        |
+| Commit               | String    | Optional | Change number for the artifact                                                        |
+| MaxCommit            | String    | Optional | Maximum commit for the artifact                                                       |
+| Name                 | String    | Optional | Name of the artifact                                                                  |
+| Type                 | String    | Required | The artifact type. Determines the permissions and expiration policy for the artifact. |
+| Keys                 | String    | Optional | Keys for the artifact                                                                 |
+| OutputDir            | String    | Required | Output directory for                                                                  |
+| Platform             | String    | Required | The platform the artifact is created for                                              |
+| Project              | String    | Required | The project this artifact is created for TODO: Can this be defaulted?                 |
+| Host                 | String    | Optional | The platform the artifact is created for                                              |
+| Namespace            | String    | Optional | The platform the artifact is created for                                              |
+| AccessToken          | String    | Optional | The access token to use                                                               |
+| AllowMultipleMatches | Boolean   | Optional | Set this to use the latest match if multiple artifacts are possible matches           |
+| If                   | Condition | Optional | Whether to execute this task. It is ignored if this condition evaluates to false.     |
+
 ### SanitizeReceipt
 
 Task that tags build products and/or runtime dependencies by reading from *.target files.
