Documentation for build tools

diff --git a/Engine/Documentation/Source/Programming/UnrealBuildSystem/Configuration/BuildConfigProperties/BuildConfigProperties.INT.udn b/Engine/Documentation/Source/Programming/UnrealBuildSystem/Configuration/BuildConfigProperties/BuildConfigProperties.INT.udn
--- a/Engine/Documentation/Source/Programming/UnrealBuildSystem/Configuration/BuildConfigProperties/BuildConfigProperties.INT.udn
+++ b/Engine/Documentation/Source/Programming/UnrealBuildSystem/Configuration/BuildConfigProperties/BuildConfigProperties.INT.udn
@@ -52,11 +52,12 @@ $ bUseUnityBuild : Whether to unify C++ code into larger files for faster compil
 
 $ bForceUnityBuild : Whether to force C++ source files to be combined into larger files for faster compilation.
 
-$ UnsafeTypeCastWarningLevel : Indicates what warning/error level to treat unsafe type casts as on platforms that support it (e.g., double->float or int64->int32)
+$ UnsafeTypeCastWarningLevel : Indicates what warning/error level to treat unsafe type casts as on platforms that support it (e.g., double->float or int64->int32) MSVC - https://learn.microsoft.com/en-us/cpp/error-messages/compiler-warnings/compiler-warning-levels-3-and-4-c4244 4244 - conversion from 'type1' to 'type2', possible loss of data 44244 - Note: The extra 4 is not a typo, /wLXXXX sets warning XXXX to level L https://learn.microsoft.com/en-us/cpp/error-messages/compiler-warnings/compiler-warning-level-1-c4838 4838 - conversion from 'type1' to 'type2' requires a narrowing conversion 44838 - Note: The extra 4 is not a typo, /wLXXXX sets warning XXXX to level L Clang - https://clang.llvm.org/docs/DiagnosticsReference.html#wfloat-conversion https://clang.llvm.org/docs/DiagnosticsReference.html#wimplicit-float-conversion https://clang.llvm.org/docs/DiagnosticsReference.html#wimplicit-int-conversion https://clang.llvm.org/docs/DiagnosticsReference.html#wc-11-narrowing 
+    * To enable (too many hits right now) https://clang.llvm.org/docs/DiagnosticsReference.html#wshorten-64-to-32 https://clang.llvm.org/docs/DiagnosticsReference.html#wsign-conversion
 
-$ UndefinedIdentifierWarningLevel : Indicates what warning/error level to treat undefined identifiers in conditional expressions.
+$ UndefinedIdentifierWarningLevel : Indicates what warning/error level to treat undefined identifiers in conditional expressions. MSVC - https://learn.microsoft.com/en-us/cpp/error-messages/compiler-warnings/compiler-warning-level-4-c4668 4668 - 'symbol' is not defined as a preprocessor macro, replacing with '0' for 'directives' 44668 - Note: The extra 4 is not a typo, /wLXXXX sets warning XXXX to level L Clang - https://clang.llvm.org/docs/DiagnosticsReference.html#wundef
 
-$ SwitchUnhandledEnumeratorWarningLevel : Indicates what warning/error level to treat unhandled enumerators in switches on enumeration-typed values.
+$ SwitchUnhandledEnumeratorWarningLevel : Indicates what warning/error level to treat unhandled enumerators in switches on enumeration-typed values. MSVC - 4061 - https://learn.microsoft.com/en-us/cpp/error-messages/compiler-warnings/compiler-warning-level-4-c4061 44061 - Note: The extra 4 is not a typo, /wLXXXX sets warning XXXX to level L Clang - https://clang.llvm.org/docs/DiagnosticsReference.html#wswitch-enum
 
 $ PCHPerformanceIssueWarningLevel : Indicates what warning/error level to treat potential PCH performance issues.
 
