Documentation for build tools

diff --git a/Engine/Documentation/Source/Programming/UnrealBuildSystem/Configuration/BuildConfigProperties/BuildConfigProperties.INT.udn b/Engine/Documentation/Source/Programming/UnrealBuildSystem/Configuration/BuildConfigProperties/BuildConfigProperties.INT.udn
--- a/Engine/Documentation/Source/Programming/UnrealBuildSystem/Configuration/BuildConfigProperties/BuildConfigProperties.INT.udn
+++ b/Engine/Documentation/Source/Programming/UnrealBuildSystem/Configuration/BuildConfigProperties/BuildConfigProperties.INT.udn
@@ -367,6 +367,8 @@ $ bClangTimeTrace : (Experimental) Appends the -ftime-trace argument to the comm
 
 $ bCompilerTrace : Outputs compile timing information so that it can be analyzed.
 
+$ bEnableExperimentalRemoteDebugging : Whether to enable the experimental remote debugging configuration.
+
 $ bSetResourceVersions : If enabled will set the ProductVersion embeded in windows executables and dlls to contain BUILT_FROM_CHANGELIST and BuildVersion Enabled by default for all precompiled and Shipping configurations. Regardless of this setting, the versions from Build.version will be available via the BuildSettings module Note: Embedding these versions will cause resource files to be recompiled whenever changelist is updated which will cause binaries to relink
 
 $ InlineFunctionExpansionLevel : Which level to use for Inline Function Expansion when TargetRules.bUseInlining is enabled
@@ -379,6 +381,20 @@ $ bStrictConformanceMode : Enables strict standard conformance mode (/permissive
 
 $ bDisableVolatileMetadata : Volatile Metadata is enabled by default and improves x64 emulation on arm64, but may come at a small perfomance cost (/volatileMetadata-).
 
+$ bPerProjectRemoteDebugPaths : When disabled, the remote debug directory is expected to be the root of a staged build. When enabled, the project name and target will be appended to the remote debug directory, allowing multiple projects to be more easily debugged. ...this can be used alongside RunUAT BuildCookRun [...] -archive -archivedirectory=\\remotepc\Share\MyProject\
+
+$ RemoteDebugMachineX64 : Remote machine name for Windows X64 remote debugging
+
+$ RemoteDebugDirectoryX64 : Local shared folder path on the remote Windows X64 remote machine to deploy to.
+
+$ RemoteDebugMachineArm64 : Remote machine name for Windows Arm64 remote debugging
+
+$ RemoteDebugDirectoryArm64 : Local shared folder path on the remote Windows Arm64 remote machine to deploy to.
+
+$ RemoteDebugMachineArm64ec : Remote machine name for Windows Arm64ec remote debugging
+
+$ RemoteDebugDirectoryArm64ec : Local shared folder path on the remote Windows Arm64ec remote machine to deploy to.
+
 ### ModuleConfiguration
 
 $ DisableMergingModuleAndGeneratedFilesInUnityFiles : List of modules that disables merging module and generated cpp files in the same unity files.
