Documentation for build tools

diff --git a/Engine/Documentation/Source/Programming/UnrealBuildSystem/Configuration/BuildConfigProperties/BuildConfigProperties.INT.udn b/Engine/Documentation/Source/Programming/UnrealBuildSystem/Configuration/BuildConfigProperties/BuildConfigProperties.INT.udn
--- a/Engine/Documentation/Source/Programming/UnrealBuildSystem/Configuration/BuildConfigProperties/BuildConfigProperties.INT.udn
+++ b/Engine/Documentation/Source/Programming/UnrealBuildSystem/Configuration/BuildConfigProperties/BuildConfigProperties.INT.udn
@@ -321,7 +321,7 @@ $ bPGONoExtraCounters : If specified along with -PGOProfile, prevent the usage o
 
 $ bSampleBasedPGO : If specified along with -PGOProfile, use sample-based PGO instead of instrumented. Currently Intel oneAPI 2024.0+ only.
 
-$ Compiler : Version of the compiler toolchain to use on Windows platform. A value of "default" will be changed to a specific version at UBT start up.
+$ DefaultPreferredCompilers : Default order of preference when Compiler is set to Default. If populated and none listed are available this will cause a warning.
 
 $ CompilerVersion : The specific compiler version to use. This may be a specific version number (for example, "14.13.26128"), the string "Latest" to select the newest available version, or the string "Preview" to select the newest available preview version. By default, and if it is available, we use the toolchain version indicated by WindowsPlatform.DefaultToolChainVersion (otherwise, we use the latest version).
 
@@ -373,8 +373,12 @@ $ bSetResourceVersions : If enabled will set the ProductVersion embeded in windo
 
 $ InlineFunctionExpansionLevel : Which level to use for Inline Function Expansion when TargetRules.bUseInlining is enabled
 
+$ Compiler : Version of the compiler toolchain to use on Windows platform. A value of "default" will be changed to a specific version at UBT start up.
+
 $ ToolChain : Version of the toolchain to use on Windows platform when a non-msvc Compiler is in use, to locate include paths etc.
 
+$ PreferredCompilers : Order of preference when Compiler is set to Default. If populated and none listed are available this will cause a warning.
+
 $ ToolchainVersionWarningLevel : Warning level when reporting toolchains that are not in the preferred version list
 
 $ bStrictConformanceMode : Enables strict standard conformance mode (/permissive-).
