Documentation for build tools

diff --git a/Engine/Documentation/Source/Programming/UnrealBuildSystem/Configuration/BuildConfigProperties/BuildConfigProperties.INT.udn b/Engine/Documentation/Source/Programming/UnrealBuildSystem/Configuration/BuildConfigProperties/BuildConfigProperties.INT.udn
--- a/Engine/Documentation/Source/Programming/UnrealBuildSystem/Configuration/BuildConfigProperties/BuildConfigProperties.INT.udn
+++ b/Engine/Documentation/Source/Programming/UnrealBuildSystem/Configuration/BuildConfigProperties/BuildConfigProperties.INT.udn
@@ -97,6 +97,8 @@ $ bUseDebugLiveCodingConsole : Whether to enable support for live coding
 
 $ bUseXGEController : Whether the XGE controller worker and modules should be included in the engine build. These are required for distributed shader compilation using the XGE interception interface.
 
+$ bAlwaysUseUnityForGeneratedFiles : Whether to force generated files to be placed in a unity file even when unity files are otherwise disabled. It is not recommended to disable this unless debugging issues with UHT file generation.
+
 $ bUseAdaptiveUnityBuild : Use a heuristic to determine which files are currently being iterated on and exclude them from unity blobs, result in faster incremental compile times. The current implementation uses the read-only flag to distinguish the working set, assuming that files will be made writable by the source control system if they are being modified. This is true for Perforce, but not for Git.
 
 $ bAdaptiveUnityDisablesOptimizations : Disable optimization for files that are in the adaptive non-unity working set.
@@ -645,6 +647,8 @@ $ PrimaryProjectName : Name of the primary project file -- for example, the base
 
 $ bPrimaryProjectNameFromFolder : If true, sets the primary project name according to the name of the folder it is in.
 
+$ bPrimaryProjectNameAppendDriveLetter : If true, append the drive letter to the primary project name.
+
 $ bIncludeTestAndShippingConfigs : Whether we should include configurations for "Test" and "Shipping" in generated projects. Pass "-NoShippingConfigs" to disable this.
 
 $ bIncludeDebugConfigs : Whether we should include configurations for "Debug" and "DebugGame" in generated projects. Pass "-NoDebugConfigs" to disable this.
@@ -687,6 +691,8 @@ $ PrimaryProjectName : Name of the primary project file -- for example, the base
 
 $ bPrimaryProjectNameFromFolder : If true, sets the primary project name according to the name of the folder it is in.
 
+$ bPrimaryProjectNameAppendDriveLetter : If true, append the drive letter to the primary project name.
+
 ### CMakefileGenerator
 
 $ bIncludeDocumentation : True if we should include documentation in the generated projects.
@@ -713,6 +719,8 @@ $ PrimaryProjectName : Name of the primary project file -- for example, the base
 
 $ bPrimaryProjectNameFromFolder : If true, sets the primary project name according to the name of the folder it is in.
 
+$ bPrimaryProjectNameAppendDriveLetter : If true, append the drive letter to the primary project name.
+
 ### CodeLiteGenerator
 
 $ bIncludeDocumentation : True if we should include documentation in the generated projects.
@@ -739,6 +747,8 @@ $ PrimaryProjectName : Name of the primary project file -- for example, the base
 
 $ bPrimaryProjectNameFromFolder : If true, sets the primary project name according to the name of the folder it is in.
 
+$ bPrimaryProjectNameAppendDriveLetter : If true, append the drive letter to the primary project name.
+
 ### EddieProjectFileGenerator
 
 $ bIncludeDocumentation : True if we should include documentation in the generated projects.
@@ -765,6 +775,8 @@ $ PrimaryProjectName : Name of the primary project file -- for example, the base
 
 $ bPrimaryProjectNameFromFolder : If true, sets the primary project name according to the name of the folder it is in.
 
+$ bPrimaryProjectNameAppendDriveLetter : If true, append the drive letter to the primary project name.
+
 ### KDevelopGenerator
 
 $ bIncludeDocumentation : True if we should include documentation in the generated projects.
@@ -791,6 +803,8 @@ $ PrimaryProjectName : Name of the primary project file -- for example, the base
 
 $ bPrimaryProjectNameFromFolder : If true, sets the primary project name according to the name of the folder it is in.
 
+$ bPrimaryProjectNameAppendDriveLetter : If true, append the drive letter to the primary project name.
+
 ### MakefileGenerator
 
 $ bIncludeDocumentation : True if we should include documentation in the generated projects.
@@ -817,6 +831,8 @@ $ PrimaryProjectName : Name of the primary project file -- for example, the base
 
 $ bPrimaryProjectNameFromFolder : If true, sets the primary project name according to the name of the folder it is in.
 
+$ bPrimaryProjectNameAppendDriveLetter : If true, append the drive letter to the primary project name.
+
 ### QMakefileGenerator
 
 $ bIncludeDocumentation : True if we should include documentation in the generated projects.
@@ -843,6 +859,8 @@ $ PrimaryProjectName : Name of the primary project file -- for example, the base
 
 $ bPrimaryProjectNameFromFolder : If true, sets the primary project name according to the name of the folder it is in.
 
+$ bPrimaryProjectNameAppendDriveLetter : If true, append the drive letter to the primary project name.
+
 ### RiderProjectFileGenerator
 
 $ bIncludeDocumentation : True if we should include documentation in the generated projects.
@@ -869,6 +887,8 @@ $ PrimaryProjectName : Name of the primary project file -- for example, the base
 
 $ bPrimaryProjectNameFromFolder : If true, sets the primary project name according to the name of the folder it is in.
 
+$ bPrimaryProjectNameAppendDriveLetter : If true, append the drive letter to the primary project name.
+
 ### VSCodeProjectFileGenerator
 
 $ IncludeAllFiles : Includes all files in the generated workspace.
@@ -905,6 +925,8 @@ $ PrimaryProjectName : Name of the primary project file -- for example, the base
 
 $ bPrimaryProjectNameFromFolder : If true, sets the primary project name according to the name of the folder it is in.
 
+$ bPrimaryProjectNameAppendDriveLetter : If true, append the drive letter to the primary project name.
+
 ### VCMacProjectFileGenerator
 
 $ bIncludeDocumentation : True if we should include documentation in the generated projects.
@@ -931,6 +953,8 @@ $ PrimaryProjectName : Name of the primary project file -- for example, the base
 
 $ bPrimaryProjectNameFromFolder : If true, sets the primary project name according to the name of the folder it is in.
 
+$ bPrimaryProjectNameAppendDriveLetter : If true, append the drive letter to the primary project name.
+
 ### VSWorkspaceProjectFileGenerator
 
 $ bIncludeDocumentation : True if we should include documentation in the generated projects.
@@ -957,6 +981,8 @@ $ PrimaryProjectName : Name of the primary project file -- for example, the base
 
 $ bPrimaryProjectNameFromFolder : If true, sets the primary project name according to the name of the folder it is in.
 
+$ bPrimaryProjectNameAppendDriveLetter : If true, append the drive letter to the primary project name.
+
 ### VCProjectFileGenerator
 
 $ Version : The version of Visual Studio to generate project files for.
@@ -997,6 +1023,8 @@ $ PrimaryProjectName : Name of the primary project file -- for example, the base
 
 $ bPrimaryProjectNameFromFolder : If true, sets the primary project name according to the name of the folder it is in.
 
+$ bPrimaryProjectNameAppendDriveLetter : If true, append the drive letter to the primary project name.
+
 ### XcodeProjectFileGenerator
 
 $ bIncludeDocumentation : True if we should include documentation in the generated projects.
@@ -1023,6 +1051,8 @@ $ PrimaryProjectName : Name of the primary project file -- for example, the base
 
 $ bPrimaryProjectNameFromFolder : If true, sets the primary project name according to the name of the folder it is in.
 
+$ bPrimaryProjectNameAppendDriveLetter : If true, append the drive letter to the primary project name.
+
 ### SourceFileWorkingSet
 
 $ Provider : Sets the provider to use for determining the working set.
diff --git a/Engine/Documentation/Source/Programming/UnrealBuildSystem/TargetFiles/TargetFilesProperties/TargetFilesProperties.INT.udn b/Engine/Documentation/Source/Programming/UnrealBuildSystem/TargetFiles/TargetFilesProperties/TargetFilesProperties.INT.udn
--- a/Engine/Documentation/Source/Programming/UnrealBuildSystem/TargetFiles/TargetFilesProperties/TargetFilesProperties.INT.udn
+++ b/Engine/Documentation/Source/Programming/UnrealBuildSystem/TargetFiles/TargetFilesProperties/TargetFilesProperties.INT.udn
@@ -364,6 +364,8 @@ $ bUseUnityBuild (Boolean): Whether to unify C++ code into larger files for fast
 
 $ bForceUnityBuild (Boolean): Whether to force C++ source files to be combined into larger files for faster compilation.
 
+$ bAlwaysUseUnityForGeneratedFiles (Boolean): Whether to force generated files to be placed in a unity file even when unity files are otherwise disabled. It is not recommended to disable this unless debugging issues with UHT file generation.
+
 $ bMergeModuleAndGeneratedUnityFiles (Boolean): Whether to merge module and generated unity files for faster compilation.
 
 $ bUseAdaptiveUnityBuild (Boolean): Use a heuristic to determine which files are currently being iterated on and exclude them from unity blobs, result in faster incremental compile times. The current implementation uses the read-only flag to distinguish the working set, assuming that files will be made writable by the source control system if they are being modified. This is true for Perforce, but not for Git.
