Documentation for build tools

diff --git a/Engine/Documentation/Source/Programming/UnrealBuildSystem/TargetFiles/TargetFilesProperties/TargetFilesProperties.INT.udn b/Engine/Documentation/Source/Programming/UnrealBuildSystem/TargetFiles/TargetFilesProperties/TargetFilesProperties.INT.udn
--- a/Engine/Documentation/Source/Programming/UnrealBuildSystem/TargetFiles/TargetFilesProperties/TargetFilesProperties.INT.udn
+++ b/Engine/Documentation/Source/Programming/UnrealBuildSystem/TargetFiles/TargetFilesProperties/TargetFilesProperties.INT.udn
@@ -294,7 +294,7 @@ $ bWithLiveCoding (Boolean): Whether to enable support for live coding
 
 $ bUseDebugLiveCodingConsole (Boolean): Whether to enable support for live coding
 
-$ bWithDirectXMath (Boolean): Whether to enable support for DirectX Math LWC_TODO: No longer supported. Needs removing.
+$ bWithDirectXMath (Boolean): Obsolete: whether to enable support for DirectX Math Not compatible with LWC so has been disabled since UE 5.
 
 $ bWithFixedTimeStepSupport (Boolean): Whether to enable the support for FixedTimeStep in the engine
 
