Documentation for build tools

diff --git a/Engine/Documentation/Source/Programming/UnrealBuildSystem/Configuration/BuildConfigProperties/BuildConfigProperties.INT.udn b/Engine/Documentation/Source/Programming/UnrealBuildSystem/Configuration/BuildConfigProperties/BuildConfigProperties.INT.udn
--- a/Engine/Documentation/Source/Programming/UnrealBuildSystem/Configuration/BuildConfigProperties/BuildConfigProperties.INT.udn
+++ b/Engine/Documentation/Source/Programming/UnrealBuildSystem/Configuration/BuildConfigProperties/BuildConfigProperties.INT.udn
@@ -485,12 +485,12 @@ $ Cache : Address of the uba cache service. Will automatically use cache if conn
 
 $ WriteCache : Set cache to write instead of fetch
 
+$ AllowDetour : Disable all detouring making all actions run outside uba
+
 $ CacheMaxWorkers : Max number of cache download tasks that can execute in parallel
 
 $ ReportCacheMissReason : Report reason a cache miss happened. Useful when searching for determinism/portability issues
 
-$ bDisableHorde : Disable horde all together
-
 ### TestTargetRules
 
 $ bCompileChaosVisualDebuggerSupport : Whether to compile in Chaos Visual Debugger (CVD) support features to record the state of the physics simulation
@@ -583,6 +583,10 @@ $ Encryption : Encryption to use for agent/compute communication. Note that UBA
 
 $ UBASentryUrl : Sentry URL to send box data to. Optional.
 
+### Horde2
+
+$ Server : Uri of the Horde server
+
 ### XGE
 
 $ bAllowOverVpn : When set to false, XGE will not be enabled when running connected to the coordinator over VPN. Configure VPN-assigned subnets via the VpnSubnets parameter.
