Documentation for build tools

diff --git a/Engine/Documentation/Source/Programming/Development/BuildGraph/ScriptAnatomy/Tasks/PredefinedTasks/BuildGraphScriptPredefinedTasks.INT.udn b/Engine/Documentation/Source/Programming/Development/BuildGraph/ScriptAnatomy/Tasks/PredefinedTasks/BuildGraphScriptPredefinedTasks.INT.udn
--- a/Engine/Documentation/Source/Programming/Development/BuildGraph/ScriptAnatomy/Tasks/PredefinedTasks/BuildGraphScriptPredefinedTasks.INT.udn
+++ b/Engine/Documentation/Source/Programming/Development/BuildGraph/ScriptAnatomy/Tasks/PredefinedTasks/BuildGraphScriptPredefinedTasks.INT.udn
@@ -184,6 +184,9 @@ Uploads an artifact to Cloud DDC
 | Force                    | Boolean   | Optional | Force the blobs to be uploaded no matter if they exist or not                                                                  |
 | BaselineBranch           | String    | Required | The streamid used for your mainline, used to build content as patches on top of this branch if no other baselines exist.       |
 | MaxBlobUploadParallelism | Integer   | Optional | Set to upload more blobs in parallel, defaults to number of cores. Set to -1 to go as wide as possible.                        |
+| AllowMultipart           | Boolean   | Optional | Enable to use multipart endpoints if valuable                                                                                  |
+| UseZenCli                | Boolean   | Optional | Use the zen cli to upload the build. Generally faster but experimental.                                                        |
+| HttpVersion              | String    | Optional | Set the explicit http version to use. None to use http handshaking.                                                            |
 | If                       | Condition | Optional | Whether to execute this task. It is ignored if this condition evaluates to false.                                              |
 
 ### CsCompile
@@ -604,6 +607,9 @@ Retrieves an artifact from Cloud DDC
 | Namespace            | String    | Optional | The platform the artifact is created for                                              |
 | AccessToken          | String    | Optional | The access token to use                                                               |
 | AllowMultipleMatches | Boolean   | Optional | Set this to use the latest match if multiple artifacts are possible matches           |
+| AllowMultipart       | Boolean   | Optional | Enable to use multipart endpoints if valuable                                         |
+| UseZenCli            | Boolean   | Optional | Use the zen cli to upload the build. Generally faster but experimental.               |
+| HttpVersion          | String    | Optional | Set the explicit http version to use. None to use http handshaking.                   |
 | If                   | Condition | Optional | Whether to execute this task. It is ignored if this condition evaluates to false.     |
 
 ### SanitizeReceipt
diff --git a/Engine/Documentation/Source/Programming/UnrealBuildSystem/Configuration/BuildConfigProperties/BuildConfigProperties.INT.udn b/Engine/Documentation/Source/Programming/UnrealBuildSystem/Configuration/BuildConfigProperties/BuildConfigProperties.INT.udn
--- a/Engine/Documentation/Source/Programming/UnrealBuildSystem/Configuration/BuildConfigProperties/BuildConfigProperties.INT.udn
+++ b/Engine/Documentation/Source/Programming/UnrealBuildSystem/Configuration/BuildConfigProperties/BuildConfigProperties.INT.udn
@@ -491,6 +491,8 @@ $ CacheMaxWorkers : Max number of cache download tasks that can execute in paral
 
 $ ReportCacheMissReason : Report reason a cache miss happened. Useful when searching for determinism/portability issues
 
+$ bDisableHorde : Disable horde all together
+
 ### TestTargetRules
 
 $ bCompileChaosVisualDebuggerSupport : Whether to compile in Chaos Visual Debugger (CVD) support features to record the state of the physics simulation
