Documentation for build tools

diff --git a/Engine/Documentation/Source/Programming/Development/BuildGraph/ScriptAnatomy/Tasks/PredefinedTasks/BuildGraphScriptPredefinedTasks.INT.udn b/Engine/Documentation/Source/Programming/Development/BuildGraph/ScriptAnatomy/Tasks/PredefinedTasks/BuildGraphScriptPredefinedTasks.INT.udn
--- a/Engine/Documentation/Source/Programming/Development/BuildGraph/ScriptAnatomy/Tasks/PredefinedTasks/BuildGraphScriptPredefinedTasks.INT.udn
+++ b/Engine/Documentation/Source/Programming/Development/BuildGraph/ScriptAnatomy/Tasks/PredefinedTasks/BuildGraphScriptPredefinedTasks.INT.udn
@@ -577,16 +577,17 @@ Renames a file, or group of files.
 
 Retrieves an artifact from Horde
 
-|           |           |          |                                                                                       |
-| --------- | --------- | -------- | ------------------------------------------------------------------------------------- |
-| Stream    | String    | Optional | Stream containing the artifact                                                        |
-| Commit    | String    | Optional | Change number for the artifact                                                        |
-| MaxCommit | String    | Optional | Maximum commit for the artifact                                                       |
-| Name      | String    | Optional | Name of the artifact                                                                  |
-| Type      | String    | Required | The artifact type. Determines the permissions and expiration policy for the artifact. |
-| Keys      | String    | Optional | Keys for the artifact                                                                 |
-| OutputDir | String    | Required | Output directory for                                                                  |
-| If        | Condition | Optional | Whether to execute this task. It is ignored if this condition evaluates to false.     |
+|              |           |          |                                                                                       |
+| ------------ | --------- | -------- | ------------------------------------------------------------------------------------- |
+| Stream       | String    | Optional | Stream containing the artifact                                                        |
+| Commit       | String    | Optional | Change number for the artifact                                                        |
+| MaxCommit    | String    | Optional | Maximum commit for the artifact                                                       |
+| Name         | String    | Optional | Name of the artifact                                                                  |
+| Type         | String    | Required | The artifact type. Determines the permissions and expiration policy for the artifact. |
+| Keys         | String    | Optional | Keys for the artifact                                                                 |
+| OutputDir    | String    | Required | Output directory for                                                                  |
+| ScopeToJobId | Boolean   | Optional | Append the environment variable UE_HORDE_JOBID to the artifact query                  |
+| If           | Condition | Optional | Whether to execute this task. It is ignored if this condition evaluates to false.     |
 
 ### RetrieveCloudArtifact
 
diff --git a/Engine/Documentation/Source/Programming/UnrealBuildSystem/Configuration/BuildConfigProperties/BuildConfigProperties.INT.udn b/Engine/Documentation/Source/Programming/UnrealBuildSystem/Configuration/BuildConfigProperties/BuildConfigProperties.INT.udn
--- a/Engine/Documentation/Source/Programming/UnrealBuildSystem/Configuration/BuildConfigProperties/BuildConfigProperties.INT.udn
+++ b/Engine/Documentation/Source/Programming/UnrealBuildSystem/Configuration/BuildConfigProperties/BuildConfigProperties.INT.udn
@@ -99,7 +99,7 @@ $ bUseXGEController : Whether the XGE controller worker and modules should be in
 
 $ bAlwaysUseUnityForGeneratedFiles : Whether to force generated files to be placed in a unity file even when unity files are otherwise disabled. It is not recommended to disable this unless debugging issues with UHT file generation.
 
-$ bUseAdaptiveUnityBuild : Use a heuristic to determine which files are currently being iterated on and exclude them from unity blobs, result in faster incremental compile times. The current implementation uses the read-only flag to distinguish the working set, assuming that files will be made writable by the source control system if they are being modified. This is true for Perforce, but not for Git.
+$ bUseAdaptiveUnityBuild : Use a heuristic to determine which files are currently being iterated on and exclude them from unity blobs. This results in faster incremental compile times. For Perforce repositories, the current implementation uses the read-only flag to distinguish the working set, assuming that files will be made writable by the source control system if they are being modified. For Git repositories, the implementation uses the git status command. Source code archives downloaded from Git as .zip files are not supported.
 
 $ bAdaptiveUnityDisablesOptimizations : Disable optimization for files that are in the adaptive non-unity working set.
 
diff --git a/Engine/Documentation/Source/Programming/UnrealBuildSystem/TargetFiles/TargetFilesProperties/TargetFilesProperties.INT.udn b/Engine/Documentation/Source/Programming/UnrealBuildSystem/TargetFiles/TargetFilesProperties/TargetFilesProperties.INT.udn
--- a/Engine/Documentation/Source/Programming/UnrealBuildSystem/TargetFiles/TargetFilesProperties/TargetFilesProperties.INT.udn
+++ b/Engine/Documentation/Source/Programming/UnrealBuildSystem/TargetFiles/TargetFilesProperties/TargetFilesProperties.INT.udn
@@ -368,7 +368,7 @@ $ bAlwaysUseUnityForGeneratedFiles (Boolean): Whether to force generated files t
 
 $ bMergeModuleAndGeneratedUnityFiles (Boolean): Whether to merge module and generated unity files for faster compilation.
 
-$ bUseAdaptiveUnityBuild (Boolean): Use a heuristic to determine which files are currently being iterated on and exclude them from unity blobs, result in faster incremental compile times. The current implementation uses the read-only flag to distinguish the working set, assuming that files will be made writable by the source control system if they are being modified. This is true for Perforce, but not for Git.
+$ bUseAdaptiveUnityBuild (Boolean): Use a heuristic to determine which files are currently being iterated on and exclude them from unity blobs. This results in faster incremental compile times. For Perforce repositories, the current implementation uses the read-only flag to distinguish the working set, assuming that files will be made writable by the source control system if they are being modified. For Git repositories, the implementation uses the git status command. Source code archives downloaded from Git as .zip files are not supported.
 
 $ bAdaptiveUnityDisablesOptimizations (Boolean): Disable optimization for files that are in the adaptive non-unity working set.
 
