Documentation for build tools

diff --git a/Engine/Documentation/Source/Programming/UnrealBuildSystem/Configuration/BuildConfigProperties/BuildConfigProperties.INT.udn b/Engine/Documentation/Source/Programming/UnrealBuildSystem/Configuration/BuildConfigProperties/BuildConfigProperties.INT.udn
--- a/Engine/Documentation/Source/Programming/UnrealBuildSystem/Configuration/BuildConfigProperties/BuildConfigProperties.INT.udn
+++ b/Engine/Documentation/Source/Programming/UnrealBuildSystem/Configuration/BuildConfigProperties/BuildConfigProperties.INT.udn
@@ -427,7 +427,9 @@ $ bDisableRemote : When set to true, UBA will not use any remote help
 
 $ bForceBuildAllRemote : When set to true, UBA will force all actions that can be built remotely to be built remotely. This will hang if there are no remote agents available
 
-$ bForcedRetry : When set to true, actions that fail locally with UBA will be retried without UBA.
+$ bAllowRetry : When set to true, actions that fail locally with certain error codes will retry without uba.
+
+$ bForcedRetry : When set to true, all actions that fail locally with UBA will be retried without UBA.
 
 $ bForcedRetryRemote : When set to true, actions that fail remotely with UBA will be retried locally with UBA.
 
