Documentation for build tools

diff --git a/Engine/Documentation/Source/Programming/Development/BuildGraph/ScriptAnatomy/Tasks/PredefinedTasks/BuildGraphScriptPredefinedTasks.INT.udn b/Engine/Documentation/Source/Programming/Development/BuildGraph/ScriptAnatomy/Tasks/PredefinedTasks/BuildGraphScriptPredefinedTasks.INT.udn
--- a/Engine/Documentation/Source/Programming/Development/BuildGraph/ScriptAnatomy/Tasks/PredefinedTasks/BuildGraphScriptPredefinedTasks.INT.udn
+++ b/Engine/Documentation/Source/Programming/Development/BuildGraph/ScriptAnatomy/Tasks/PredefinedTasks/BuildGraphScriptPredefinedTasks.INT.udn
@@ -879,6 +879,8 @@ Imports a snapshot from a specified destination to Zen.
 | OverridePlatformCookedDir | String        | Optional | Optional. Where to look for the ue.projectstore file. The pattern {Platform} can be used for importing multiple platforms at once. |
 | ForceImport               | Boolean       | Optional | Optional. If true, force import of the oplog (corresponds to --force on the Zen command line).                                     |
 | AsyncImport               | Boolean       | Optional | Optional. If true, import oplogs asynchronously (corresponds to --async on the Zen command line).                                  |
+| RemoteZenHost             | String        | Optional | Optional. Remote zenserver host to import snapshots into.                                                                          |
+| RemoteZenPort             | Integer       | Optional | Optional. Port on remote host which zenserver is listening on.                                                                     |
 | If                        | Condition     | Optional | Whether to execute this task. It is ignored if this condition evaluates to false.                                                  |
 
 ### ZenLaunch
