Documentation for build tools

diff --git a/Engine/Documentation/Source/Programming/Development/BuildGraph/ScriptAnatomy/Tasks/PredefinedTasks/BuildGraphScriptPredefinedTasks.INT.udn b/Engine/Documentation/Source/Programming/Development/BuildGraph/ScriptAnatomy/Tasks/PredefinedTasks/BuildGraphScriptPredefinedTasks.INT.udn
--- a/Engine/Documentation/Source/Programming/Development/BuildGraph/ScriptAnatomy/Tasks/PredefinedTasks/BuildGraphScriptPredefinedTasks.INT.udn
+++ b/Engine/Documentation/Source/Programming/Development/BuildGraph/ScriptAnatomy/Tasks/PredefinedTasks/BuildGraphScriptPredefinedTasks.INT.udn
@@ -180,7 +180,7 @@ Uploads an artifact to Cloud DDC
 | StreamId                 | String        | Optional | Stream containing the artifact.                                                                                                |
 | Commit                   | String        | Optional | Commit for the uploaded artifact.                                                                                              |
 | Files                    | FileSpec      | Required | Files to include in the artifact.                                                                                              |
-| Metadata                 | String        | Optional | Other metadata for the artifact, separated by semicolons.                                                                      |
+| Metadata                 | String        | Optional | Other metadata for the artifact, separated by semicolons. Use key=value for each field.                                        |
 | Force                    | Boolean       | Optional | Force the blobs to be uploaded no matter if they exist or not                                                                  |
 | BaselineBranch           | String        | Optional | The streamid used for your mainline, used to build content as patches on top of this branch if no other baselines exist.       |
 | MaxBlobUploadParallelism | Integer       | Optional | Set to upload more blobs in parallel, defaults to number of cores. Set to -1 to go as wide as possible.                        |
