Documentation for build tools

diff --git a/Engine/Documentation/Source/Programming/Development/BuildGraph/ScriptAnatomy/Tasks/PredefinedTasks/BuildGraphScriptPredefinedTasks.INT.udn b/Engine/Documentation/Source/Programming/Development/BuildGraph/ScriptAnatomy/Tasks/PredefinedTasks/BuildGraphScriptPredefinedTasks.INT.udn
--- a/Engine/Documentation/Source/Programming/Development/BuildGraph/ScriptAnatomy/Tasks/PredefinedTasks/BuildGraphScriptPredefinedTasks.INT.udn
+++ b/Engine/Documentation/Source/Programming/Development/BuildGraph/ScriptAnatomy/Tasks/PredefinedTasks/BuildGraphScriptPredefinedTasks.INT.udn
@@ -169,6 +169,7 @@ Uploads an artifact to Cloud DDC
 |                          |               |          |                                                                                                                                |
 | ------------------------ | ------------- | -------- | ------------------------------------------------------------------------------------------------------------------------------ |
 | Name                     | String        | Required | Name of the artifact                                                                                                           |
+| HordeArtifactName        | String        | Optional | Optional override of the name of the artifact as presented in Horde, defaults to name if not set                               |
 | Type                     | String        | Required | The artifact type. Determines the permissions and expiration policy for the artifact.                                          |
 | Description              | String        | Optional | Description for the artifact. Will be shown through the Horde dashboard.                                                       |
 | Platform                 | String        | Required | The platform the artifact is created for                                                                                       |
@@ -185,7 +186,6 @@ Uploads an artifact to Cloud DDC
 | BaselineBranch           | String        | Optional | The streamid used for your mainline, used to build content as patches on top of this branch if no other baselines exist.       |
 | MaxBlobUploadParallelism | Integer       | Optional | Set to upload more blobs in parallel, defaults to number of cores. Set to -1 to go as wide as possible.                        |
 | AllowMultipart           | Boolean       | Optional | Enable to use multipart endpoints if valuable                                                                                  |
-| UseZenCli                | Boolean       | Optional | Use the zen cli to upload the build. Generally faster but experimental.                                                        |
 | HttpVersion              | String        | Optional | Set the explicit http version to use. None to use http handshaking.                                                            |
 | If                       | Condition     | Optional | Whether to execute this task. It is ignored if this condition evaluates to false.                                              |
 
