Documentation for build tools

diff --git a/Engine/Documentation/Source/Programming/Development/BuildGraph/ScriptAnatomy/Tasks/PredefinedTasks/BuildGraphScriptPredefinedTasks.INT.udn b/Engine/Documentation/Source/Programming/Development/BuildGraph/ScriptAnatomy/Tasks/PredefinedTasks/BuildGraphScriptPredefinedTasks.INT.udn
--- a/Engine/Documentation/Source/Programming/Development/BuildGraph/ScriptAnatomy/Tasks/PredefinedTasks/BuildGraphScriptPredefinedTasks.INT.udn
+++ b/Engine/Documentation/Source/Programming/Development/BuildGraph/ScriptAnatomy/Tasks/PredefinedTasks/BuildGraphScriptPredefinedTasks.INT.udn
@@ -609,7 +609,6 @@ Retrieves an artifact from Cloud DDC
 | AccessToken          | String        | Optional | The access token to use                                                               |
 | AllowMultipleMatches | Boolean       | Optional | Set this to use the latest match if multiple artifacts are possible matches           |
 | AllowMultipart       | Boolean       | Optional | Enable to use multipart endpoints if valuable                                         |
-| UseZenCli            | Boolean       | Optional | Use the zen cli to upload the build. Generally faster but experimental.               |
 | HttpVersion          | String        | Optional | Set the explicit http version to use. None to use http handshaking.                   |
 | If                   | Condition     | Optional | Whether to execute this task. It is ignored if this condition evaluates to false.     |
 
