Documentation for build tools

diff --git a/Engine/Documentation/Source/Programming/Development/BuildGraph/ScriptAnatomy/Tasks/PredefinedTasks/BuildGraphScriptPredefinedTasks.INT.udn b/Engine/Documentation/Source/Programming/Development/BuildGraph/ScriptAnatomy/Tasks/PredefinedTasks/BuildGraphScriptPredefinedTasks.INT.udn
--- a/Engine/Documentation/Source/Programming/Development/BuildGraph/ScriptAnatomy/Tasks/PredefinedTasks/BuildGraphScriptPredefinedTasks.INT.udn
+++ b/Engine/Documentation/Source/Programming/Development/BuildGraph/ScriptAnatomy/Tasks/PredefinedTasks/BuildGraphScriptPredefinedTasks.INT.udn
@@ -173,7 +173,8 @@ Uploads an artifact to Cloud DDC
 | Type                     | String        | Required | The artifact type. Determines the permissions and expiration policy for the artifact.                                          |
 | Description              | String        | Optional | Description for the artifact. Will be shown through the Horde dashboard.                                                       |
 | Platform                 | String        | Required | The platform the artifact is created for                                                                                       |
-| Project                  | FileReference | Optional | The project this artifact is created for                                                                                       |
+| Project                  | FileReference | Optional | The uproject this artifact is created for                                                                                      |
+| ProjectNameOverride      | String        | Optional | Override the project name, defaults to the filename of the uproject                                                            |
 | Host                     | String        | Optional | The platform the artifact is created for                                                                                       |
 | Namespace                | String        | Optional | The platform the artifact is created for                                                                                       |
 | AccessToken              | String        | Optional | The access token to use                                                                                                        |
