Documentation for build tools

diff --git a/Engine/Documentation/Source/Programming/UnrealBuildSystem/Configuration/BuildConfigProperties/BuildConfigProperties.INT.udn b/Engine/Documentation/Source/Programming/UnrealBuildSystem/Configuration/BuildConfigProperties/BuildConfigProperties.INT.udn
--- a/Engine/Documentation/Source/Programming/UnrealBuildSystem/Configuration/BuildConfigProperties/BuildConfigProperties.INT.udn
+++ b/Engine/Documentation/Source/Programming/UnrealBuildSystem/Configuration/BuildConfigProperties/BuildConfigProperties.INT.udn
@@ -375,6 +375,8 @@ $ bSetResourceVersions : If enabled will set the ProductVersion embeded in windo
 
 $ InlineFunctionExpansionLevel : Which level to use for Inline Function Expansion when TargetRules.bUseInlining is enabled
 
+$ bDynamicDebugging : If specificed Visual Studio Dynamic Debugging support will be enabled, unless the compiler or linker is not MSVC or LTCG is enabled
+
 $ Compiler : Version of the compiler toolchain to use on Windows platform. A value of "default" will be changed to a specific version at UBT start up.
 
 $ ToolChain : Version of the toolchain to use on Windows platform when a non-msvc Compiler is in use, to locate include paths etc.
