Documentation for build tools

diff --git a/Engine/Documentation/Source/Programming/UnrealBuildSystem/Configuration/BuildConfigProperties/BuildConfigProperties.INT.udn b/Engine/Documentation/Source/Programming/UnrealBuildSystem/Configuration/BuildConfigProperties/BuildConfigProperties.INT.udn
--- a/Engine/Documentation/Source/Programming/UnrealBuildSystem/Configuration/BuildConfigProperties/BuildConfigProperties.INT.udn
+++ b/Engine/Documentation/Source/Programming/UnrealBuildSystem/Configuration/BuildConfigProperties/BuildConfigProperties.INT.udn
@@ -495,6 +495,8 @@ $ CacheMaxWorkers : Max number of cache download tasks that can execute in paral
 
 $ ReportCacheMissReason : Report reason a cache miss happened. Useful when searching for determinism/portability issues
 
+$ CompressionLevel : Compression level. Options are None, SuperFast, VeryFast, Fast, Normal, Optimal1 to Optimal5
+
 $ bDisableHorde : Disable horde all together
 
 ### TestTargetRules
