Documentation for build tools

diff --git a/Engine/Documentation/Source/Programming/Development/BuildGraph/ScriptAnatomy/Tasks/PredefinedTasks/BuildGraphScriptPredefinedTasks.INT.udn b/Engine/Documentation/Source/Programming/Development/BuildGraph/ScriptAnatomy/Tasks/PredefinedTasks/BuildGraphScriptPredefinedTasks.INT.udn
--- a/Engine/Documentation/Source/Programming/Development/BuildGraph/ScriptAnatomy/Tasks/PredefinedTasks/BuildGraphScriptPredefinedTasks.INT.udn
+++ b/Engine/Documentation/Source/Programming/Development/BuildGraph/ScriptAnatomy/Tasks/PredefinedTasks/BuildGraphScriptPredefinedTasks.INT.udn
@@ -881,6 +881,7 @@ Imports a snapshot from a specified destination to Zen.
 | AsyncImport               | Boolean       | Optional | Optional. If true, import oplogs asynchronously (corresponds to --async on the Zen command line).                                  |
 | RemoteZenHost             | String        | Optional | Optional. Remote zenserver host to import snapshots into.                                                                          |
 | RemoteZenPort             | Integer       | Optional | Optional. Port on remote host which zenserver is listening on.                                                                     |
+| DestinationProjectId      | String        | Optional | Optional. Destination project ID to import snapshots into.                                                                         |
 | If                        | Condition     | Optional | Whether to execute this task. It is ignored if this condition evaluates to false.                                                  |
 
 ### ZenLaunch
