Documentation for build tools

diff --git a/Engine/Documentation/Source/Programming/UnrealBuildSystem/Configuration/BuildConfigProperties/BuildConfigProperties.INT.udn b/Engine/Documentation/Source/Programming/UnrealBuildSystem/Configuration/BuildConfigProperties/BuildConfigProperties.INT.udn
--- a/Engine/Documentation/Source/Programming/UnrealBuildSystem/Configuration/BuildConfigProperties/BuildConfigProperties.INT.udn
+++ b/Engine/Documentation/Source/Programming/UnrealBuildSystem/Configuration/BuildConfigProperties/BuildConfigProperties.INT.udn
@@ -287,6 +287,8 @@ $ bUsePerFileIntellisense : Whether to generate per-file intellisense data.
 
 $ bEditorDependsOnShaderCompileWorker : Whether to include a dependency on ShaderCompileWorker when generating project files for the editor.
 
+$ bShowTimeline : If set and tool execution was successful, then display an unreal build tool script execution timeline summary. If unset or the tool execution failed, print the same information silently to the log.
+
 $ TempDirectory : If set TMP\TEMP will be overidden to this directory, each process will create a unique subdirectory in this folder.
 
 $ bDeleteTempDirectory : If set the application temp directory will be deleted on exit, only when running with a single instance mutex.
