Documentation for build tools

diff --git a/Engine/Documentation/Source/Programming/UnrealBuildSystem/Configuration/BuildConfigProperties/BuildConfigProperties.INT.udn b/Engine/Documentation/Source/Programming/UnrealBuildSystem/Configuration/BuildConfigProperties/BuildConfigProperties.INT.udn
--- a/Engine/Documentation/Source/Programming/UnrealBuildSystem/Configuration/BuildConfigProperties/BuildConfigProperties.INT.udn
+++ b/Engine/Documentation/Source/Programming/UnrealBuildSystem/Configuration/BuildConfigProperties/BuildConfigProperties.INT.udn
@@ -239,6 +239,8 @@ $ CStandard : Which C standard to use for compiling this target
 
 $ MinCpuArchX64 : Direct the compiler to generate AVX instructions wherever SSE or AVX intrinsics are used, on the x64 platforms that support it. Ignored for arm64. Note that by enabling this you are changing the minspec for the target platform, and the resultant executable will crash on machines without AVX support.
 
+$ MinArm64CpuTarget : Select a minimum ARM64 target CPU. This will change what ARMv8(9)-A ISA + extension code will be compiled for. Changing it from default will guarantee not to work on all ARM64 CPUs. Currently supported by clang compiler only
+
 $ ActionStallReportTime : Number of second of no completed actions to trigger an action stall report. If zero, stall reports will not be enabled.
 
 $ ActionStallTerminateTime : Number of second of no completed actions to trigger to terminate the queue. If zero, force termination not be enabled.
diff --git a/Engine/Documentation/Source/Programming/UnrealBuildSystem/ModuleFiles/ModuleFilesProperties/ModuleFilesProperties.INT.udn b/Engine/Documentation/Source/Programming/UnrealBuildSystem/ModuleFiles/ModuleFilesProperties/ModuleFilesProperties.INT.udn
--- a/Engine/Documentation/Source/Programming/UnrealBuildSystem/ModuleFiles/ModuleFilesProperties/ModuleFilesProperties.INT.udn
+++ b/Engine/Documentation/Source/Programming/UnrealBuildSystem/ModuleFiles/ModuleFilesProperties/ModuleFilesProperties.INT.udn
@@ -44,6 +44,10 @@ $ AllowedRestrictedFolders (List&lt;String&gt;): List of folders which are allow
 
 $ AliasRestrictedFolders (Dictionary&lt;String, String&gt;): Set of aliased restricted folder references
 
+$ bAlwaysExportStructs (Boolean): If true, autogenerated functions of USTRUCTS are always exported
+
+$ bAlwaysExportEnums (Boolean): If true, autogenerated functions of UENUMS are always exported
+
 $ PublicIncludePathModuleNames (List&lt;String&gt;): List of modules names (no path needed) with header files that our module's public headers needs access to, but we don't need to "import" or link against.
 
 $ PublicDependencyModuleNames (List&lt;String&gt;): List of public dependency module names (no path needed) (automatically does the private/public include). These are modules that are required by our public source files.
diff --git a/Engine/Documentation/Source/Programming/UnrealBuildSystem/TargetFiles/TargetFilesProperties/TargetFilesProperties.INT.udn b/Engine/Documentation/Source/Programming/UnrealBuildSystem/TargetFiles/TargetFilesProperties/TargetFilesProperties.INT.udn
--- a/Engine/Documentation/Source/Programming/UnrealBuildSystem/TargetFiles/TargetFilesProperties/TargetFilesProperties.INT.udn
+++ b/Engine/Documentation/Source/Programming/UnrealBuildSystem/TargetFiles/TargetFilesProperties/TargetFilesProperties.INT.udn
@@ -580,6 +580,8 @@ $ CStandard (CStandardVersion): Which C standard to use for compiling this targe
 
 $ MinCpuArchX64 (MinimumCpuArchitectureX64): Direct the compiler to generate AVX instructions wherever SSE or AVX intrinsics are used, on the x64 platforms that support it. Ignored for arm64. Note that by enabling this you are changing the minspec for the target platform, and the resultant executable will crash on machines without AVX support.
 
+$ MinArm64CpuTarget (Arm64TargetCpuType): Select a minimum ARM64 target CPU. This will change what ARMv8(9)-A ISA + extension code will be compiled for. Changing it from default will guarantee not to work on all ARM64 CPUs. Currently supported by clang compiler only
+
 $ BuildVersion (String): The build version string
 
 $ LinkType (TargetLinkType): Specifies how to link modules in this target (monolithic or modular). This is currently protected for backwards compatibility. Call the GetLinkType() accessor until support for the deprecated ShouldCompileMonolithic() override has been removed.
