Documentation for build tools

diff --git a/Engine/Documentation/Source/Programming/Development/BuildGraph/ScriptAnatomy/Tasks/PredefinedTasks/BuildGraphScriptPredefinedTasks.INT.udn b/Engine/Documentation/Source/Programming/Development/BuildGraph/ScriptAnatomy/Tasks/PredefinedTasks/BuildGraphScriptPredefinedTasks.INT.udn
--- a/Engine/Documentation/Source/Programming/Development/BuildGraph/ScriptAnatomy/Tasks/PredefinedTasks/BuildGraphScriptPredefinedTasks.INT.udn
+++ b/Engine/Documentation/Source/Programming/Development/BuildGraph/ScriptAnatomy/Tasks/PredefinedTasks/BuildGraphScriptPredefinedTasks.INT.udn
@@ -166,29 +166,31 @@ Uploads an artifact to Horde
 
 Uploads an artifact to Cloud DDC
 
-|                          |               |          |                                                                                                                                |
-| ------------------------ | ------------- | -------- | ------------------------------------------------------------------------------------------------------------------------------ |
-| Name                     | String        | Required | Name of the artifact                                                                                                           |
-| HordeArtifactName        | String        | Optional | Optional override of the name of the artifact as presented in Horde, defaults to name if not set                               |
-| Type                     | String        | Required | The artifact type. Determines the permissions and expiration policy for the artifact.                                          |
-| Description              | String        | Optional | Description for the artifact. Will be shown through the Horde dashboard.                                                       |
-| Platform                 | String        | Required | The platform the artifact is created for                                                                                       |
-| Project                  | FileReference | Optional | The uproject this artifact is created for                                                                                      |
-| ProjectNameOverride      | String        | Optional | Override the project name, defaults to the filename of the uproject                                                            |
-| Host                     | String        | Optional | The platform the artifact is created for                                                                                       |
-| Namespace                | String        | Optional | The platform the artifact is created for                                                                                       |
-| AccessToken              | String        | Optional | The access token to use                                                                                                        |
-| BaseDir                  | String        | Optional | Base path for the uploaded files. All the tagged files must be under this directory. Defaults to the workspace root directory. |
-| StreamId                 | String        | Optional | Stream containing the artifact.                                                                                                |
-| Commit                   | String        | Optional | Commit for the uploaded artifact.                                                                                              |
-| Files                    | FileSpec      | Required | Files to include in the artifact.                                                                                              |
-| Metadata                 | String        | Optional | Other metadata for the artifact, separated by semicolons. Use key=value for each field.                                        |
-| Force                    | Boolean       | Optional | Force the blobs to be uploaded no matter if they exist or not                                                                  |
-| BaselineBranch           | String        | Optional | The streamid used for your mainline, used to build content as patches on top of this branch if no other baselines exist.       |
-| MaxBlobUploadParallelism | Integer       | Optional | Set to upload more blobs in parallel, defaults to number of cores. Set to -1 to go as wide as possible.                        |
-| AllowMultipart           | Boolean       | Optional | Enable to use multipart endpoints if valuable                                                                                  |
-| HttpVersion              | String        | Optional | Set the explicit http version to use. None to use http handshaking.                                                            |
-| If                       | Condition     | Optional | Whether to execute this task. It is ignored if this condition evaluates to false.                                              |
+|                          |               |          |                                                                                                                                                                                                                          |
+| ------------------------ | ------------- | -------- | ------------------------------------------------------------------------------------------------------------------------------------------------------------------------------------------------------------------------ |
+| Name                     | String        | Required | Name of the artifact                                                                                                                                                                                                     |
+| HordeArtifactName        | String        | Optional | Optional override of the name of the artifact as presented in Horde, defaults to name if not set                                                                                                                         |
+| Type                     | String        | Required | The artifact type. Determines the permissions and expiration policy for the artifact.                                                                                                                                    |
+| Description              | String        | Optional | Description for the artifact. Will be shown through the Horde dashboard.                                                                                                                                                 |
+| Platform                 | String        | Required | The platform the artifact is created for                                                                                                                                                                                 |
+| Project                  | FileReference | Optional | The uproject this artifact is created for                                                                                                                                                                                |
+| ProjectNameOverride      | String        | Optional | Override the project name, defaults to the filename of the uproject                                                                                                                                                      |
+| Host                     | String        | Optional | The platform the artifact is created for                                                                                                                                                                                 |
+| Namespace                | String        | Optional | The platform the artifact is created for                                                                                                                                                                                 |
+| AccessToken              | String        | Optional | The access token to use                                                                                                                                                                                                  |
+| BaseDir                  | String        | Optional | Base path for the uploaded files. All the tagged files must be under this directory. Defaults to the workspace root directory.                                                                                           |
+| StreamId                 | String        | Optional | Stream containing the artifact.                                                                                                                                                                                          |
+| Commit                   | String        | Optional | Commit for the uploaded artifact.                                                                                                                                                                                        |
+| Files                    | FileSpec      | Required | Files to include in the artifact.                                                                                                                                                                                        |
+| Metadata                 | String        | Optional | Other metadata for the artifact, separated by semicolons. Use key=value for each field.                                                                                                                                  |
+| Force                    | Boolean       | Optional | Force the blobs to be uploaded no matter if they exist or not                                                                                                                                                            |
+| BaselineBranch           | String        | Optional | The streamid used for your mainline, used to build content as patches on top of this branch if no other baselines exist.                                                                                                 |
+| MaxBlobUploadParallelism | Integer       | Optional | Set to upload more blobs in parallel, defaults to number of cores. Set to -1 to go as wide as possible.                                                                                                                  |
+| AllowMultipart           | Boolean       | Optional | Enable to use multipart endpoints if valuable                                                                                                                                                                            |
+| HttpVersion              | String        | Optional | Set the explicit http version to use. None to use http handshaking.                                                                                                                                                      |
+| Verify                   | Boolean       | Optional | Verify the contents after upload, essentially downloads the entire build to make sure all data was uploaded correctly. Takes a bit of time and is generally not needed but can be useful to detect issues in the upload. |
+| Verbose                  | Boolean       | Optional | Verbose logging output, can be useful to find out more information when an issue has occured.                                                                                                                            |
+| If                       | Condition     | Optional | Whether to execute this task. It is ignored if this condition evaluates to false.                                                                                                                                        |
 
 ### CsCompile
 
