Documentation for build tools

diff --git a/Engine/Documentation/Source/Programming/UnrealBuildSystem/Configuration/BuildConfigProperties/BuildConfigProperties.INT.udn b/Engine/Documentation/Source/Programming/UnrealBuildSystem/Configuration/BuildConfigProperties/BuildConfigProperties.INT.udn
--- a/Engine/Documentation/Source/Programming/UnrealBuildSystem/Configuration/BuildConfigProperties/BuildConfigProperties.INT.udn
+++ b/Engine/Documentation/Source/Programming/UnrealBuildSystem/Configuration/BuildConfigProperties/BuildConfigProperties.INT.udn
@@ -263,13 +263,13 @@ $ bUseDSYMFiles : Enables the generation of .dsym files. This can be disabled to
 
 $ bSkipClangValidation : Disables clang build verification checks on static libraries
 
-$ bEnableAddressSanitizer : Enables address sanitizer (ASan). Only supported for Visual Studio 2019 16.7.0 and up.
+$ bEnableAddressSanitizer : Enables address sanitizer (ASan)
 
-$ bEnableLibFuzzer : Enables LibFuzzer. Only supported for Visual Studio 2022 17.0.0 and up.
+$ bEnableLibFuzzer : Enables LibFuzzer.
 
 $ bEnableThreadSanitizer : Enables thread sanitizer (TSan).
 
-$ bEnableUndefinedBehaviorSanitizer : Enables undefined behavior sanitizer (UBSan).
+$ bEnableUndefinedBehaviorSanitizer : Enables undefined behavior sanitizer (UBSan)
 
 $ bEnableMemorySanitizer : Enables memory sanitizer (MSan)
 
@@ -393,6 +393,8 @@ $ bStrictConformanceMode : Enables strict standard conformance mode (/permissive
 
 $ bDisableVolatileMetadata : Volatile Metadata is enabled by default and improves x64 emulation on arm64, but may come at a small perfomance cost (/volatileMetadata-).
 
+$ StructMemberAlignment : Controls how the members of a structure are packed into memory and specifies the same packing for all structures in a module. See https://learn.microsoft.com/en-us/cpp/build/reference/zp-struct-member-alignment
+
 $ bPerProjectRemoteDebugPaths : When disabled, the remote debug directory is expected to be the root of a staged build. When enabled, the project name and target will be appended to the remote debug directory, allowing multiple projects to be more easily debugged. ...this can be used alongside RunUAT BuildCookRun [...] -archive -archivedirectory=\\remotepc\Share\MyProject\
 
 $ RemoteDebugMachineX64 : Remote machine name for Windows X64 remote debugging
