Documentation for build tools

diff --git a/Engine/Documentation/Source/Programming/Development/BuildGraph/ScriptAnatomy/Tasks/PredefinedTasks/BuildGraphScriptPredefinedTasks.INT.udn b/Engine/Documentation/Source/Programming/Development/BuildGraph/ScriptAnatomy/Tasks/PredefinedTasks/BuildGraphScriptPredefinedTasks.INT.udn
--- a/Engine/Documentation/Source/Programming/Development/BuildGraph/ScriptAnatomy/Tasks/PredefinedTasks/BuildGraphScriptPredefinedTasks.INT.udn
+++ b/Engine/Documentation/Source/Programming/Development/BuildGraph/ScriptAnatomy/Tasks/PredefinedTasks/BuildGraphScriptPredefinedTasks.INT.udn
@@ -190,6 +190,7 @@ Uploads an artifact to Cloud DDC
 | HttpVersion              | String        | Optional | Set the explicit http version to use. None to use http handshaking.                                                                                                                                                      |
 | Verify                   | Boolean       | Optional | Verify the contents after upload, essentially downloads the entire build to make sure all data was uploaded correctly. Takes a bit of time and is generally not needed but can be useful to detect issues in the upload. |
 | Verbose                  | Boolean       | Optional | Verbose logging output, can be useful to find out more information when an issue has occured.                                                                                                                            |
+| BoostWorkers             | Boolean       | Optional | Increase the number of worker threads used by zen, may cause machine to be less responsive but will generally improve upload times                                                                                       |
 | If                       | Condition     | Optional | Whether to execute this task. It is ignored if this condition evaluates to false.                                                                                                                                        |
 
 ### CsCompile
@@ -596,24 +597,25 @@ Retrieves an artifact from Horde
 
 Retrieves an artifact from Cloud DDC
 
-|                      |               |          |                                                                                       |
-| -------------------- | ------------- | -------- | ------------------------------------------------------------------------------------- |
-| StreamId             | String        | Optional | Stream containing the artifact                                                        |
-| Commit               | String        | Optional | Change number for the artifact                                                        |
-| MaxCommit            | String        | Optional | Maximum commit for the artifact                                                       |
-| Name                 | String        | Optional | Name of the artifact                                                                  |
-| Type                 | String        | Required | The artifact type. Determines the permissions and expiration policy for the artifact. |
-| Keys                 | String        | Optional | Keys for the artifact                                                                 |
-| OutputDir            | String        | Required | Output directory for                                                                  |
-| Platform             | String        | Required | The platform the artifact is created for                                              |
-| Project              | FileReference | Required | The path to the uproject this artifact is created for                                 |
-| Host                 | String        | Optional | The platform the artifact is created for                                              |
-| Namespace            | String        | Optional | The platform the artifact is created for                                              |
-| AccessToken          | String        | Optional | The access token to use                                                               |
-| AllowMultipleMatches | Boolean       | Optional | Set this to use the latest match if multiple artifacts are possible matches           |
-| AllowMultipart       | Boolean       | Optional | Enable to use multipart endpoints if valuable                                         |
-| HttpVersion          | String        | Optional | Set the explicit http version to use. None to use http handshaking.                   |
-| If                   | Condition     | Optional | Whether to execute this task. It is ignored if this condition evaluates to false.     |
+|                      |               |          |                                                                                                                                      |
+| -------------------- | ------------- | -------- | ------------------------------------------------------------------------------------------------------------------------------------ |
+| StreamId             | String        | Optional | Stream containing the artifact                                                                                                       |
+| Commit               | String        | Optional | Change number for the artifact                                                                                                       |
+| MaxCommit            | String        | Optional | Maximum commit for the artifact                                                                                                      |
+| Name                 | String        | Optional | Name of the artifact                                                                                                                 |
+| Type                 | String        | Required | The artifact type. Determines the permissions and expiration policy for the artifact.                                                |
+| Keys                 | String        | Optional | Keys for the artifact                                                                                                                |
+| OutputDir            | String        | Required | Output directory for                                                                                                                 |
+| Platform             | String        | Required | The platform the artifact is created for                                                                                             |
+| Project              | FileReference | Required | The path to the uproject this artifact is created for                                                                                |
+| Host                 | String        | Optional | The platform the artifact is created for                                                                                             |
+| Namespace            | String        | Optional | The platform the artifact is created for                                                                                             |
+| AccessToken          | String        | Optional | The access token to use                                                                                                              |
+| AllowMultipleMatches | Boolean       | Optional | Set this to use the latest match if multiple artifacts are possible matches                                                          |
+| AllowMultipart       | Boolean       | Optional | Enable to use multipart endpoints if valuable                                                                                        |
+| HttpVersion          | String        | Optional | Set the explicit http version to use. None to use http handshaking.                                                                  |
+| BoostWorkers         | Boolean       | Optional | Increase the number of worker threads used by zen, may cause machine to be less responsive but will generally improve download times |
+| If                   | Condition     | Optional | Whether to execute this task. It is ignored if this condition evaluates to false.                                                    |
 
 ### SanitizeReceipt
 
