Documentation for build tools

diff --git a/Engine/Documentation/Source/Programming/UnrealBuildSystem/ModuleFiles/ModuleFilesProperties/ModuleFilesProperties.INT.udn b/Engine/Documentation/Source/Programming/UnrealBuildSystem/ModuleFiles/ModuleFilesProperties/ModuleFilesProperties.INT.udn
--- a/Engine/Documentation/Source/Programming/UnrealBuildSystem/ModuleFiles/ModuleFilesProperties/ModuleFilesProperties.INT.udn
+++ b/Engine/Documentation/Source/Programming/UnrealBuildSystem/ModuleFiles/ModuleFilesProperties/ModuleFilesProperties.INT.udn
@@ -112,6 +112,8 @@ $ GenerateHeaderFuncs (List&lt;ValueTuple&lt;String, Action&lt;ILogger, Director
 
 $ SwiftInteropHeader (String): If set, this module will generate a bridging header between C++ and Swift, using this input header
 
+$ FPassPlugins (List&lt;String&gt;): A List of files to be used as a LLVM FPass-Plugins for all clang compilation units of the current module. 'clang -fpass-plugin=[FPassPlugin]'
+
 $ Name (String): Name of this module
 
 $ Type (ModuleType): Type of module
