Documentation for build tools

diff --git a/Engine/Documentation/Source/Programming/UnrealBuildSystem/Configuration/BuildConfigProperties/BuildConfigProperties.INT.udn b/Engine/Documentation/Source/Programming/UnrealBuildSystem/Configuration/BuildConfigProperties/BuildConfigProperties.INT.udn
--- a/Engine/Documentation/Source/Programming/UnrealBuildSystem/Configuration/BuildConfigProperties/BuildConfigProperties.INT.udn
+++ b/Engine/Documentation/Source/Programming/UnrealBuildSystem/Configuration/BuildConfigProperties/BuildConfigProperties.INT.udn
@@ -431,6 +431,10 @@ $ bCompileChaosVisualDebuggerSupport : Whether to compile in Chaos Visual Debugg
 
 $ bStoreObjFilesCompressed : Store object (.obj) compressed on disk. Requires uba to do link step where it will decompress obj files again
 
+$ Providers : List of provider names for remote connections.
+
+$ BuildMachineProviders : List of provider names for remote connections
+
 $ bDisableRemote : When set to true, UBA will not use any remote help
 
 $ bForceBuildAllRemote : When set to true, UBA will force all actions that can be built remotely to be built remotely. This will hang if there are no remote agents available
@@ -597,10 +601,6 @@ $ Encryption : Encryption to use for agent/compute communication. Note that UBA
 
 $ UBASentryUrl : Sentry URL to send box data to. Optional.
 
-### Horde2
-
-$ Server : Uri of the Horde server
-
 ### XGE
 
 $ bAllowOverVpn : When set to false, XGE will not be enabled when running connected to the coordinator over VPN. Configure VPN-assigned subnets via the VpnSubnets parameter.
