Documentation for build tools

diff --git a/Engine/Documentation/Source/Programming/UnrealBuildSystem/Configuration/BuildConfigProperties/BuildConfigProperties.INT.udn b/Engine/Documentation/Source/Programming/UnrealBuildSystem/Configuration/BuildConfigProperties/BuildConfigProperties.INT.udn
--- a/Engine/Documentation/Source/Programming/UnrealBuildSystem/Configuration/BuildConfigProperties/BuildConfigProperties.INT.udn
+++ b/Engine/Documentation/Source/Programming/UnrealBuildSystem/Configuration/BuildConfigProperties/BuildConfigProperties.INT.udn
@@ -59,6 +59,8 @@ $ UndefinedIdentifierWarningLevel : Indicates what warning/error level to treat
 
 $ SwitchUnhandledEnumeratorWarningLevel : Indicates what warning/error level to treat unhandled enumerators in switches on enumeration-typed values. MSVC - 4061 - https://learn.microsoft.com/en-us/cpp/error-messages/compiler-warnings/compiler-warning-level-4-c4061 44061 - Note: The extra 4 is not a typo, /wLXXXX sets warning XXXX to level L Clang - https://clang.llvm.org/docs/DiagnosticsReference.html#wswitch-enum
 
+$ DeprecationWarningLevel : Indicates what warning/error level to treat usages of functions, methods, or variables, that are marked as deprecated. Clang - https://clang.llvm.org/docs/DiagnosticsReference.html#wdeprecated-declarations
+
 $ PCHPerformanceIssueWarningLevel : Indicates what warning/error level to treat potential PCH performance issues.
 
 $ ModuleUnsupportedWarningLevel : How to treat module unsupported validation messages
@@ -73,8 +75,6 @@ $ ModuleIncludeSubdirectoryWarningLevel : How to treat unnecessary module sub-di
 
 $ DefaultWarningLevel : Default treatment of uncategorized warnings
 
-$ DeprecationWarningLevel : Level to report deprecation warnings as errors
-
 $ bShowIncludes : Print out files that are included by each source file
 
 $ bDebugBuildsActuallyUseDebugCRT : Enables the debug C++ runtime (CRT) for debug builds. By default we always use the release runtime, since the debug version isn't particularly useful when debugging Unreal Engine projects, and linking against the debug CRT libraries forces our third party library dependencies to also be compiled using the debug CRT (and often perform more slowly). Often it can be inconvenient to require a separate copy of the debug versions of third party static libraries simply so that you can debug your program's code.
