Documentation for build tools

diff --git a/Engine/Documentation/Source/Programming/Development/BuildGraph/ScriptAnatomy/Tasks/PredefinedTasks/BuildGraphScriptPredefinedTasks.INT.udn b/Engine/Documentation/Source/Programming/Development/BuildGraph/ScriptAnatomy/Tasks/PredefinedTasks/BuildGraphScriptPredefinedTasks.INT.udn
--- a/Engine/Documentation/Source/Programming/Development/BuildGraph/ScriptAnatomy/Tasks/PredefinedTasks/BuildGraphScriptPredefinedTasks.INT.udn
+++ b/Engine/Documentation/Source/Programming/Development/BuildGraph/ScriptAnatomy/Tasks/PredefinedTasks/BuildGraphScriptPredefinedTasks.INT.udn
@@ -2,7 +2,7 @@ Availability: NoPublish
 Title: BuildGraph Predefined Tasks
 Crumbs: %ROOT%, Programming, Programming/Development, Programming/Development/BuildGraph, Programming/Development/BuildGraph/BuildGraphScriptTasks
 Description: This is a procedurally generated markdown page.
-version: 5.6
+version: 5.7
 parent:Programming/Development/BuildGraph/BuildGraphScriptTasks
 
 ### AgeStore
diff --git a/Engine/Documentation/Source/Programming/UnrealBuildSystem/Configuration/BuildConfigProperties/BuildConfigProperties.INT.udn b/Engine/Documentation/Source/Programming/UnrealBuildSystem/Configuration/BuildConfigProperties/BuildConfigProperties.INT.udn
--- a/Engine/Documentation/Source/Programming/UnrealBuildSystem/Configuration/BuildConfigProperties/BuildConfigProperties.INT.udn
+++ b/Engine/Documentation/Source/Programming/UnrealBuildSystem/Configuration/BuildConfigProperties/BuildConfigProperties.INT.udn
@@ -2,7 +2,7 @@ Availability: NoPublish
 Title: Build Configuration Properties Page
 Crumbs:
 Description: This is a procedurally generated markdown page.
-Version: 5.6
+Version: 5.7
 
 ### BuildConfiguration
 
@@ -431,10 +431,20 @@ $ bCompileChaosVisualDebuggerSupport : Whether to compile in Chaos Visual Debugg
 
 $ bStoreObjFilesCompressed : Store object (.obj) compressed on disk. Requires uba to do link step where it will decompress obj files again
 
+$ Cache : Address of the uba cache service. Will automatically use cache if connected
+
+$ WriteCache : Set cache to write, expects one of [True, False, BuildMachineOnly]
+
+$ RequireVFS : Set cache to require VFS to be enabled
+
 $ Providers : List of provider names for remote connections.
 
 $ BuildMachineProviders : List of provider names for remote connections
 
+$ CacheProviders : List of provider names for remote cache.
+
+$ BuildMachineCacheProviders : List of provider names for remote cache
+
 $ bDisableRemote : When set to true, UBA will not use any remote help
 
 $ bForceBuildAllRemote : When set to true, UBA will force all actions that can be built remotely to be built remotely. This will hang if there are no remote agents available
@@ -495,14 +505,12 @@ $ ActionsOutputFile : Write yaml file with all actions that are queued for build
 
 $ bDetailedLog : Set to true to see more info about what is happening inside uba and also log output from agents
 
-$ Cache : Address of the uba cache service. Will automatically use cache if connected
-
-$ WriteCache : Set cache to write instead of fetch
-
 $ AllowDetour : Disable all detouring making all actions run outside uba
 
 $ CacheMaxWorkers : Max number of cache download tasks that can execute in parallel
 
+$ CacheShuffle : Shuffle the list of cache servers before each read attempt
+
 $ ReportCacheMissReason : Report reason a cache miss happened. Useful when searching for determinism/portability issues
 
 $ CompressionLevel : Compression level. Options are None, SuperFast, VeryFast, Fast, Normal, Optimal1 to Optimal5
diff --git a/Engine/Documentation/Source/Programming/UnrealBuildSystem/ModuleFiles/ModuleFilesProperties/ModuleFilesProperties.INT.udn b/Engine/Documentation/Source/Programming/UnrealBuildSystem/ModuleFiles/ModuleFilesProperties/ModuleFilesProperties.INT.udn
--- a/Engine/Documentation/Source/Programming/UnrealBuildSystem/ModuleFiles/ModuleFilesProperties/ModuleFilesProperties.INT.udn
+++ b/Engine/Documentation/Source/Programming/UnrealBuildSystem/ModuleFiles/ModuleFilesProperties/ModuleFilesProperties.INT.udn
@@ -2,7 +2,7 @@ Availability: NoPublish
 Title: Build Configuration Properties Page
 Crumbs:
 Description: This is a procedurally generated markdown page.
-Version: 5.6
+Version: 5.7
 
 ### Read-Only Properties
 
diff --git a/Engine/Documentation/Source/Programming/UnrealBuildSystem/TargetFiles/TargetFilesProperties/TargetFilesProperties.INT.udn b/Engine/Documentation/Source/Programming/UnrealBuildSystem/TargetFiles/TargetFilesProperties/TargetFilesProperties.INT.udn
--- a/Engine/Documentation/Source/Programming/UnrealBuildSystem/TargetFiles/TargetFilesProperties/TargetFilesProperties.INT.udn
+++ b/Engine/Documentation/Source/Programming/UnrealBuildSystem/TargetFiles/TargetFilesProperties/TargetFilesProperties.INT.udn
@@ -2,7 +2,7 @@ Availability: NoPublish
 Title: Build Configuration Properties Page
 Crumbs:
 Description: This is a procedurally generated markdown page.
-Version: 5.6
+Version: 5.7
 
 ### Read-Only Properties
 
