Documentation for build tools

diff --git a/Engine/Documentation/Source/Programming/UnrealBuildSystem/Configuration/BuildConfigProperties/BuildConfigProperties.INT.udn b/Engine/Documentation/Source/Programming/UnrealBuildSystem/Configuration/BuildConfigProperties/BuildConfigProperties.INT.udn
--- a/Engine/Documentation/Source/Programming/UnrealBuildSystem/Configuration/BuildConfigProperties/BuildConfigProperties.INT.udn
+++ b/Engine/Documentation/Source/Programming/UnrealBuildSystem/Configuration/BuildConfigProperties/BuildConfigProperties.INT.udn
@@ -1021,6 +1021,8 @@ $ ExcludedFilePaths : Semi-colon separated list of paths that should not be adde
 
 $ bBuildUBTInDebug : Forces UBT to be built in debug configuration, regardless of the solution configuration
 
+$ bHeadersAsClCompile : Whether to treat headers as ClCompile. May cause intellisense issues but allows the VS Compile command to build headers.
+
 $ bBuildLiveCodingConsole : Whether to include a dependency on LiveCodingConsole when building targets that support live coding.
 
 $ bMakeProjectPerTarget : Whether to generate a project file for each individual target, and not include e.g. Editor/Client/Server in the Configuration.
