Documentation for build tools

diff --git a/Engine/Documentation/Source/Programming/UnrealBuildSystem/TargetFiles/TargetFilesProperties/TargetFilesProperties.INT.udn b/Engine/Documentation/Source/Programming/UnrealBuildSystem/TargetFiles/TargetFilesProperties/TargetFilesProperties.INT.udn
--- a/Engine/Documentation/Source/Programming/UnrealBuildSystem/TargetFiles/TargetFilesProperties/TargetFilesProperties.INT.udn
+++ b/Engine/Documentation/Source/Programming/UnrealBuildSystem/TargetFiles/TargetFilesProperties/TargetFilesProperties.INT.udn
@@ -28,6 +28,8 @@ $ DebugInfoLineTablesOnlyPlugins (HashSet&lt;String&gt;): Plugins that should em
 
 $ MergePlugins (List&lt;String&gt;): Experimental: List of plugins (and their dependencies) to each merge into separate libraries. Requires bMergeModules to be enabled
 
+$ MergePluginsIgnored (List&lt;String&gt;): Experimental: List of plugins (and their dependencies) to ignore when processing merged libraries. Requires bMergeModules to be enabled Ignored plugins will still be compiled, but if they're kept in the base executable, their dependencies will not be forced out of merged library with them. This should be used carefully, for organisational plugins known to have dependencies they don't actually need at runtime.
+
 $ MergePluginsShared (Dictionary&lt;String, HashSet&lt;String&gt;&gt;): Experimental: List of merged plugins to move common dependencies to a shared library, requires MergePlugins. Can be chained
 
 $ MergePluginsLaunch (List&lt;String&gt;): Experimental: List of merged plugins to move to the main executable, requires MergePlugins. "Engine" for all shared engine modules and "Common" for shared project modules.
