Documentation for build tools

diff --git a/Engine/Documentation/Source/Programming/UnrealBuildSystem/TargetFiles/TargetFilesProperties/TargetFilesProperties.INT.udn b/Engine/Documentation/Source/Programming/UnrealBuildSystem/TargetFiles/TargetFilesProperties/TargetFilesProperties.INT.udn
--- a/Engine/Documentation/Source/Programming/UnrealBuildSystem/TargetFiles/TargetFilesProperties/TargetFilesProperties.INT.udn
+++ b/Engine/Documentation/Source/Programming/UnrealBuildSystem/TargetFiles/TargetFilesProperties/TargetFilesProperties.INT.udn
@@ -30,10 +30,14 @@ $ MergePlugins (List&lt;String&gt;): Experimental: List of plugins (and their de
 
 $ MergePluginsIgnored (List&lt;String&gt;): Experimental: List of plugins (and their dependencies) to ignore when processing merged libraries. Requires bMergeModules to be enabled Ignored plugins will still be compiled, but if they're kept in the base executable, their dependencies will not be forced out of merged library with them. This should be used carefully, for organisational plugins known to have dependencies they don't actually need at runtime.
 
+$ MergePluginsExcluded (List&lt;String&gt;): Experimental: List of plugins (and their dependencies) to explicitly exclude when processing merged libraries. Requires bMergeModules to be enabled Ignored plugins will still be compiled but forced into the executable.
+
 $ MergePluginsShared (Dictionary&lt;String, HashSet&lt;String&gt;&gt;): Experimental: List of merged plugins to move common dependencies to a shared library, requires MergePlugins. Can be chained
 
 $ MergePluginsLaunch (List&lt;String&gt;): Experimental: List of merged plugins to move to the main executable, requires MergePlugins. "Engine" for all shared engine modules and "Common" for shared project modules.
 
+$ MergeOnlyGameFeatures (Boolean): Experimental: Whether to merge only game feature plugins, ignoring regular plugins.
+
 $ GlobalDefinitions (List&lt;String&gt;): Macros to define globally across the whole target.
 
 $ ProjectDefinitions (List&lt;String&gt;): Macros to define across all macros in the project.
