Documentation for build tools

diff --git a/Engine/Documentation/Source/Programming/UnrealBuildSystem/Configuration/BuildConfigProperties/BuildConfigProperties.INT.udn b/Engine/Documentation/Source/Programming/UnrealBuildSystem/Configuration/BuildConfigProperties/BuildConfigProperties.INT.udn
--- a/Engine/Documentation/Source/Programming/UnrealBuildSystem/Configuration/BuildConfigProperties/BuildConfigProperties.INT.udn
+++ b/Engine/Documentation/Source/Programming/UnrealBuildSystem/Configuration/BuildConfigProperties/BuildConfigProperties.INT.udn
@@ -121,8 +121,6 @@ $ bValidateFormatStrings : Emits compilation errors for incorrect UE_LOG format
 
 $ bWarningsAsErrors : Whether to enable all warnings as errors. UE enables most warnings as errors already, but disables a few (such as deprecation warnings).
 
-$ bUndefinedIdentifierErrors : Forces the use of undefined identifiers in conditional expressions to be treated as errors.
-
 $ bRetainFramePointers : Forces frame pointers to be retained this is usually required when you want reliable callstacks e.g. mallocframeprofiler
 
 $ bUseFastMonoCalls : New Monolithic Graphics drivers have optional "fast calls" replacing various D3d functions
@@ -133,10 +131,6 @@ $ bDisableModuleNumIncludedBytesPerUnityCPPOverride : Disables overrides that ar
 
 $ bStressTestUnity : Whether to stress test the C++ unity build robustness by including all C++ files files in a project from a single unified file.
 
-$ bDetailedUnityFiles : Whether to add additional information to the unity files, such as '_of_X' in the file name. Not recommended.
-
-$ bDisableDebugInfo : Whether to globally disable debug info generation; Obsolete, please use TargetRules.DebugInfoMode instead
-
 $ DebugInfo : How much debug info should be generated. See DebugInfoMode enum for more details
 
 $ DebugInfoLineTablesOnly : True if only debug line number tables should be emitted in debug information for compilers that support doing so. Overrides TargetRules.DebugInfo See https://clang.llvm.org/docs/UsersManual.html#cmdoption-gline-tables-only for more information
@@ -241,6 +235,12 @@ $ MinCpuArchX64 : Direct the compiler to generate AVX instructions wherever SSE
 
 $ MinArm64CpuTarget : Select a minimum ARM64 target CPU. This will change what ARMv8(9)-A ISA + extension code will be compiled for. Changing it from default will guarantee not to work on all ARM64 CPUs. Currently supported by clang compiler only
 
+$ bDetailedUnityFiles : Whether to add additional information to the unity files, such as '_of_X' in the file name. Not recommended.
+
+$ bDisableDebugInfo : Whether to globally disable debug info generation; Obsolete, please use TargetRules.DebugInfoMode instead
+
+$ bUndefinedIdentifierErrors : Forces the use of undefined identifiers in conditional expressions to be treated as errors.
+
 $ ActionStallReportTime : Number of second of no completed actions to trigger an action stall report. If zero, stall reports will not be enabled.
 
 $ ActionStallTerminateTime : Number of second of no completed actions to trigger to terminate the queue. If zero, force termination not be enabled.
diff --git a/Engine/Documentation/Source/Programming/UnrealBuildSystem/ModuleFiles/ModuleFilesProperties/ModuleFilesProperties.INT.udn b/Engine/Documentation/Source/Programming/UnrealBuildSystem/ModuleFiles/ModuleFilesProperties/ModuleFilesProperties.INT.udn
--- a/Engine/Documentation/Source/Programming/UnrealBuildSystem/ModuleFiles/ModuleFilesProperties/ModuleFilesProperties.INT.udn
+++ b/Engine/Documentation/Source/Programming/UnrealBuildSystem/ModuleFiles/ModuleFilesProperties/ModuleFilesProperties.INT.udn
@@ -154,8 +154,6 @@ $ bVcRemoveUnreferencedComdat (Boolean): Whether to direct MSVC to remove unrefe
 
 $ bCodeCoverage (Boolean): Enable code coverage compilation/linking support.
 
-$ bUseAVX (Boolean): Obsolete: Direct the compiler to generate AVX instructions wherever SSE or AVX intrinsics are used, on the platforms that support it. Note that by enabling this you are changing the minspec for the PC platform, and the resultant executable will crash on machines without AVX support.
-
 $ MinCpuArchX64 (Nullable&lt;MinimumCpuArchitectureX64&gt;): Direct the compiler to generate AVX instructions wherever SSE or AVX intrinsics are used, on the x64 platforms that support it. Note that by enabling this you are changing the minspec for the PC platform, and the resultant executable will crash on machines without AVX support.
 
 $ bEnableBufferSecurityChecks (Boolean): Enable buffer security checks.  This should usually be enabled as it prevents severe security risks.
@@ -166,26 +164,10 @@ $ bEnableObjCExceptions (Boolean): Enable objective C exception handling
 
 $ bEnableObjCAutomaticReferenceCounting (Boolean): Enable objective C automatic reference counting (ARC) If you set this to true you should not use shared PCHs for this module. The engine won't be extensively using ARC in the short term Not doing this will result in a compile errors because shared PCHs were compiled with different flags than consumer
 
-$ DeterministicWarningLevel (WarningLevel): How to treat deterministic warnings (experimental).
-
 $ CppCompileWarningSettings (CppCompileWarnings): Settings object used by the module.
 
-$ ShadowVariableWarningLevel (WarningLevel): How to treat shadow variable warnings
-
 $ bWarningsAsErrors (Boolean): Whether to enable all warnings as errors. UE enables most warnings as errors already, but disables a few (such as deprecation warnings).
 
-$ UnsafeTypeCastWarningLevel (WarningLevel): How to treat unsafe implicit type cast warnings (e.g., double->float or int64->int32)
-
-$ UndefinedIdentifierWarningLevel (WarningLevel): Indicates what warning/error level to treat undefined identifiers in conditional expressions.
-
-$ bEnableUndefinedIdentifierWarnings (Boolean): Enable warnings for using undefined identifiers in #if expressions
-
-$ ModuleIncludePathWarningLevel (WarningLevel): How to treat general module include path validation messages
-
-$ ModuleIncludePrivateWarningLevel (WarningLevel): How to treat private module include path validation messages, where a module is adding an include path that exposes private headers
-
-$ ModuleIncludeSubdirectoryWarningLevel (WarningLevel): How to treat unnecessary module sub-directory include path validation messages
-
 $ bDisableStaticAnalysis (Boolean): Disable all static analysis - clang, msvc, pvs-studio.
 
 $ bStaticAnalyzerExtensions (Boolean): Enable additional analyzer extension warnings using the EspXEngine plugin. This is only supported for MSVC. See https://learn.microsoft.com/en-us/cpp/code-quality/using-the-cpp-core-guidelines-checkers This will add a large number of warnings by default. It's recommended to use StaticAnalyzerRulesets if this is enabled.
@@ -204,14 +186,10 @@ $ NumIncludedBytesPerUnityCPPOverride (Int32): Overrides Target.NumIncludedBytes
 
 $ bBuildLocallyWithSNDBS (Boolean): Module uses a #import so must be built locally when compiling with SN-DBS
 
-$ bEnableNonInlinedGenCppWarnings (Boolean): Enable warnings for when there are .gen.cpp files that could be inlined in a matching handwritten cpp file
-
 $ IsRedistributableOverride (Nullable&lt;Boolean&gt;): Redistribution override flag for this module.
 
 $ bLegalToDistributeObjectCode (Boolean): Whether the output from this module can be publicly distributed, even if it has code/ dependencies on modules that are not (i.e. CarefullyRedist, NotForLicensees, NoRedist). This should be used when you plan to release binaries but not source.
 
-$ bEnforceIWYU (Boolean): Enforce "include what you use" rules when PCHUsage is set to ExplicitOrSharedPCH; warns when monolithic headers (Engine.h, UnrealEd.h, etc...) are used, and checks that source files include their matching header first.
-
 $ IWYUSupport (IWYUSupport): Allows "include what you use" to modify the source code when run. bEnforceIWYU must be true for this variable to matter.
 
 $ bAddDefaultIncludePaths (Boolean): Whether to add all the default include paths to the module (eg. the Source/Classes folder, subfolders under Source/Public).
@@ -248,4 +226,26 @@ $ CStandard (Nullable&lt;CStandardVersion&gt;): Which standard to use for compil
 
 $ ModuleSymbolVisibility (SymbolVisibility): Control visibility of symbols
 
+$ bUseAVX (Boolean): Obsolete: Direct the compiler to generate AVX instructions wherever SSE or AVX intrinsics are used, on the platforms that support it. Note that by enabling this you are changing the minspec for the PC platform, and the resultant executable will crash on machines without AVX support.
+
+$ bEnforceIWYU (Boolean): Enforce "include what you use" rules when PCHUsage is set to ExplicitOrSharedPCH; warns when monolithic headers (Engine.h, UnrealEd.h, etc...) are used, and checks that source files include their matching header first.
+
+$ DeterministicWarningLevel (WarningLevel): How to treat deterministic warnings (experimental).
+
+$ ShadowVariableWarningLevel (WarningLevel): How to treat shadow variable warnings
+
+$ UnsafeTypeCastWarningLevel (WarningLevel): How to treat unsafe implicit type cast warnings (e.g., double->float or int64->int32)
+
+$ UndefinedIdentifierWarningLevel (WarningLevel): Indicates what warning/error level to treat undefined identifiers in conditional expressions.
+
+$ bEnableUndefinedIdentifierWarnings (Boolean): Enable warnings for using undefined identifiers in #if expressions
+
+$ ModuleIncludePathWarningLevel (WarningLevel): How to treat general module include path validation messages
+
+$ ModuleIncludePrivateWarningLevel (WarningLevel): How to treat private module include path validation messages, where a module is adding an include path that exposes private headers
+
+$ ModuleIncludeSubdirectoryWarningLevel (WarningLevel): How to treat unnecessary module sub-directory include path validation messages
+
+$ bEnableNonInlinedGenCppWarnings (Boolean): Enable warnings for when there are .gen.cpp files that could be inlined in a matching handwritten cpp file
+
 
diff --git a/Engine/Documentation/Source/Programming/UnrealBuildSystem/TargetFiles/TargetFilesProperties/TargetFilesProperties.INT.udn b/Engine/Documentation/Source/Programming/UnrealBuildSystem/TargetFiles/TargetFilesProperties/TargetFilesProperties.INT.udn
--- a/Engine/Documentation/Source/Programming/UnrealBuildSystem/TargetFiles/TargetFilesProperties/TargetFilesProperties.INT.udn
+++ b/Engine/Documentation/Source/Programming/UnrealBuildSystem/TargetFiles/TargetFilesProperties/TargetFilesProperties.INT.udn
@@ -26,6 +26,8 @@
 
 $ DebugInfoLineTablesOnlyPlugins (HashSet&lt;String&gt;): Plugins that should emit line number tables instead of full debug information for compilers that support doing so. Overrides DisableDebugInfoPlugins
 
+$ MergeAdditionalCommonModules (HashSet&lt;String&gt;): Experimental: List of modules to exclude from Engine to be placed in Common
+
 $ MergePlugins (List&lt;String&gt;): Experimental: List of plugins (and their dependencies) to each merge into separate libraries. Requires bMergeModules to be enabled
 
 $ MergePluginsIgnored (List&lt;String&gt;): Experimental: List of plugins (and their dependencies) to ignore when processing merged libraries. Requires bMergeModules to be enabled Ignored plugins will still be compiled, but if they're kept in the base executable, their dependencies will not be forced out of merged library with them. This should be used carefully, for organisational plugins known to have dependencies they don't actually need at runtime.
@@ -36,8 +38,6 @@
 
 $ MergePluginsLaunch (List&lt;String&gt;): Experimental: List of merged plugins to move to the main executable, requires MergePlugins. "Engine" for all shared engine modules and "Common" for shared project modules.
 
-$ MergeOnlyGameFeatures (Boolean): Experimental: Whether to merge only game feature plugins, ignoring regular plugins.
-
 $ GlobalDefinitions (List&lt;String&gt;): Macros to define globally across the whole target.
 
 $ ProjectDefinitions (List&lt;String&gt;): Macros to define across all macros in the project.
@@ -92,8 +92,6 @@
 
 $ OptimizeForSizeAndSpeedModules (String[]): List of modules to optimize for size and speed. This allows the target to override module optimization level Note that this may disable PCH usage if a private PCH is not provided
 
-$ OptedInModulePlatforms (UnrealTargetPlatform[]): If this is non-null, then any platforms NOT listed will not be allowed to have modules in their directories be created
-
 $ Name (String): The name of this target
 
 $ Logger (ILogger): Logger for output relating to this target. Set before the constructor is run from
@@ -124,8 +122,6 @@
 
 $ OutputFile (String): Path to the output file for the main executable, relative to the Engine or project directory. This setting is only typically useful for non-UE programs, since the engine uses paths relative to the executable to find other known folders (eg. Content).
 
-$ bUsesSteam (Boolean): Whether the target uses Steam.
-
 $ bUsesCEF3 (Boolean): Whether the target uses CEF3.
 
 $ bUsesSlate (Boolean): Whether the project uses visual Slate UI (as opposed to the low level windowing/messaging, which is always available).
@@ -150,8 +146,6 @@
 
 $ bAllowEnginePluginsEnabledByDefault (Boolean): If false, suppress loading "enabled by default" engine plugins not explicitly enabled by the project or target files.
 
-$ DisablePluginsConflictWarningLevel (WarningLevel): How to treat conflicts when a disabled plugin is being enabled by another plugin referencing it
-
 $ PakSigningKeysFile (String): Path to the set of pak signing keys to embed in the executable.
 
 $ SolutionDirectory (String): Allows a Program Target to specify it's own solution folder path.
@@ -184,32 +178,18 @@
 
 $ bAutoRTFMClosedStaticLinkage (Boolean): Whether to link closed function declarations statically.
 
-$ bCompileChaos (Boolean): Whether to compile the Chaos physics plugin.
-
-$ bUseChaos (Boolean): Whether to use the Chaos physics interface. This overrides the physx flags to disable APEX and NvCloth
-
 $ bUseChaosChecked (Boolean): Whether to compile in checked chaos features for debugging
 
 $ bUseChaosMemoryTracking (Boolean): Whether to compile in chaos memory tracking features
 
 $ bCompileChaosVisualDebuggerSupport (Boolean): Whether to compile in Chaos Visual Debugger (CVD) support features to record the state of the physics simulation
 
-$ bCustomSceneQueryStructure (Boolean): Whether scene query acceleration is done by UE. The physx scene query structure is still created, but we do not use it.
-
-$ bCompilePhysX (Boolean): Whether to include PhysX support.
-
-$ bCompileAPEX (Boolean): Whether to include PhysX APEX support.
-
-$ bCompileNvCloth (Boolean): Whether to include NvCloth.
-
 $ bCompileICU (Boolean): Whether to include ICU unicode/i18n support in Core.
 
 $ bCompileCEF3 (Boolean): Whether to compile CEF3 support.
 
 $ bCompileISPC (Boolean): Whether to compile using ISPC.
 
-$ bCompileIntelMetricsDiscovery (Boolean): Whether to compile IntelMetricsDiscovery.
-
 $ bCompilePython (Boolean): Whether to compile in python support
 
 $ bUseGameplayDebugger (Boolean): Whether to compile with WITH_GAMEPLAY_DEBUGGER enabled with all Engine's default gameplay debugger categories.
@@ -300,8 +280,6 @@
 
 $ bUseDebugLiveCodingConsole (Boolean): Whether to enable support for live coding
 
-$ bWithDirectXMath (Boolean): Obsolete: whether to enable support for DirectX Math Not compatible with LWC so has been disabled since UE 5.
-
 $ bWithFixedTimeStepSupport (Boolean): Whether to enable the support for FixedTimeStep in the engine
 
 $ bUseLoggingInShipping (Boolean): Whether to turn on logging for test/shipping builds.
@@ -326,8 +304,6 @@
 
 $ bUseExecCommandsInShipping (Boolean): Whether to turn allow exec commands for shipping builds.
 
-$ bCompileForSize (Boolean): True if we want to favor optimizing size over speed.
-
 $ OptimizationLevel (OptimizationMode): Allows to fine tune optimizations level for speed and\or code size
 
 $ FPSemantics (FPSemanticsMode): Allows setting the FP semantics.
@@ -372,8 +348,6 @@
 
 $ bAlwaysUseUnityForGeneratedFiles (Boolean): Whether to force generated files to be placed in a unity file even when unity files are otherwise disabled. It is not recommended to disable this unless debugging issues with UHT file generation.
 
-$ bMergeModuleAndGeneratedUnityFiles (Boolean): Whether to merge module and generated unity files for faster compilation.
-
 $ bUseAdaptiveUnityBuild (Boolean): Use a heuristic to determine which files are currently being iterated on and exclude them from unity blobs. This results in faster incremental compile times. For Perforce repositories, the current implementation uses the read-only flag to distinguish the working set, assuming that files will be made writable by the source control system if they are being modified. For Git repositories, the implementation uses the git status command. Source code archives downloaded from Git as .zip files are not supported.
 
 $ bAdaptiveUnityDisablesOptimizations (Boolean): Disable optimization for files that are in the adaptive non-unity working set.
@@ -398,28 +372,10 @@
 
 $ bValidateFormatStrings (Boolean): Emits compilation errors for incorrect UE_LOG format strings.
 
-$ DeprecationWarningLevel (WarningLevel): Level to report deprecation warnings as errors
-
 $ CppCompileWarningSettings (CppCompileWarnings): C++ Warnings settings object used by the target
 
-$ ShadowVariableWarningLevel (WarningLevel): Forces shadow variable warnings to be treated as errors on platforms that support it.
-
 $ bWarningsAsErrors (Boolean): Whether to enable all warnings as errors. UE enables most warnings as errors already, but disables a few (such as deprecation warnings).
 
-$ UnsafeTypeCastWarningLevel (WarningLevel): Indicates what warning/error level to treat unsafe type casts as on platforms that support it (e.g., double->float or int64->int32)
-
-$ bUndefinedIdentifierErrors (Boolean): Forces the use of undefined identifiers in conditional expressions to be treated as errors.
-
-$ UndefinedIdentifierWarningLevel (WarningLevel): Indicates what warning/error level to treat undefined identifiers in conditional expressions.
-
-$ PCHPerformanceIssueWarningLevel (WarningLevel): Indicates what warning/error level to treat potential PCH performance issues.
-
-$ ModuleIncludePathWarningLevel (WarningLevel): How to treat general module include path validation messages
-
-$ ModuleIncludePrivateWarningLevel (WarningLevel): How to treat private module include path validation messages, where a module is adding an include path that exposes private headers
-
-$ ModuleIncludeSubdirectoryWarningLevel (WarningLevel): How to treat unnecessary module sub-directory include path validation messages
-
 $ bRetainFramePointers (Boolean): Forces frame pointers to be retained this is usually required when you want reliable callstacks e.g. mallocframeprofiler
 
 $ bRetainFramePointersOverride (Nullable&lt;Boolean&gt;): Forces frame pointers to be retained this is usually required when you want reliable callstacks e.g. mallocframeprofiler
@@ -432,12 +388,8 @@
 
 $ bStressTestUnity (Boolean): Whether to stress test the C++ unity build robustness by including all C++ files files in a project from a single unified file.
 
-$ bDetailedUnityFiles (Boolean): Whether to add additional information to the unity files, such as '_of_X' in the file name. Not recommended.
-
 $ bForceDebugInfo (Boolean): Whether to force debug info to be generated.
 
-$ bDisableDebugInfo (Boolean): Whether to globally disable debug info generation; Obsolete, please use TargetRules.DebugInfoMode instead
-
 $ DebugInfo (DebugInfoMode): How much debug info should be generated. See DebugInfoMode enum for more details
 
 $ DebugInfoLineTablesOnly (DebugInfoMode): True if only debug line number tables should be emitted in debug information for compilers that support doing so. Overrides TargetRules.DebugInfo See https://clang.llvm.org/docs/UsersManual.html#cmdoption-gline-tables-only for more information
@@ -506,12 +458,8 @@
 
 $ bEnableCppModules (Boolean): Whether to enable support for C++20 modules
 
-$ bEnableCppCoroutinesForEvaluation (Boolean): Whether to enable support for C++20 coroutines This option is provided to facilitate evaluation of the feature. Expect the name of the option to change. Use of coroutines in with UE is untested and unsupported.
-
 $ bEnableProcessPriorityControl (Boolean): Whether to enable engine's ability to set process priority on runtime. This option requires some environment setup on Linux, that's why it's disabled by default. Project has to opt-in this feature as it has to guarantee correct setup.
 
-$ bUseMallocProfiler (Boolean): If true, then enable memory profiling in the build (defines USE_MALLOC_PROFILER=1 and forces bOmitFramePointers=false).
-
 $ bShaderCompilerWorkerTrace (Boolean): If true, then enable Unreal Insights (utrace) profiling in the build for the Shader Compiler Worker (defines USE_SHADER_COMPILER_WORKER_TRACE=1).
 
 $ bUseSharedPCHs (Boolean): Enables "Shared PCHs", a feature which significantly speeds up compile times by attempting to share certain PCH files between modules that UBT detects is including those PCH's header files.
@@ -560,7 +508,7 @@
 
 $ bPublicSymbolsByDefault (Boolean): Whether to expose all symbols as public by default on POSIX platforms
 
-$ bDisableInliningGenCpps (Boolean): Disable supports for inlining gen.cpps
+$ bIdenticalCodeFolding (Boolean): Enables Identical Code Folding (Clang toolchains)
 
 $ ToolChainName (String): Allows overriding the toolchain to be created for this target. This must match the name of a class declared in the UnrealBuildTool assembly.
 
@@ -596,6 +544,8 @@
 
 $ bMergeModules (Boolean): Experimental: Merge modular modules into combined libraries. Sets LinkType to Modular and enables bStripExports
 
+$ MergeOnlyGameFeatures (Boolean): Experimental: Whether to merge only game feature plugins, ignoring regular plugins.
+
 $ LaunchModuleName (String): Specifies the name of the launch module. For modular builds, this is the module that is compiled into the target's executable.
 
 $ ExportPublicHeader (String): Specifies the path to write a header containing public definitions for this target. Useful when building a DLL to be consumed by external build processes.
@@ -614,6 +564,60 @@
 
 $ GeneratedProjectName (String): When generating project files, specifies the name of the project file to use when there are multiple targets of the same type.
 
+$ OptedInModulePlatforms (UnrealTargetPlatform[]): If this is non-null, then any platforms NOT listed will not be allowed to have modules in their directories be created
+
 $ StaticAllocator (StaticAllocatorType): Statically selected allocator to use
 
+$ bUsesSteam (Boolean): Whether the target uses Steam.
 
+$ bCompileChaos (Boolean): Whether to compile the Chaos physics plugin.
+
+$ bUseChaos (Boolean): Whether to use the Chaos physics interface. This overrides the physx flags to disable APEX and NvCloth
+
+$ bCustomSceneQueryStructure (Boolean): Whether scene query acceleration is done by UE. The physx scene query structure is still created, but we do not use it.
+
+$ bCompilePhysX (Boolean): Whether to include PhysX support.
+
+$ bCompileAPEX (Boolean): Whether to include PhysX APEX support.
+
+$ bCompileNvCloth (Boolean): Whether to include NvCloth.
+
+$ bCompileIntelMetricsDiscovery (Boolean): Whether to compile IntelMetricsDiscovery.
+
+$ bWithDirectXMath (Boolean): Obsolete: whether to enable support for DirectX Math Not compatible with LWC so has been disabled since UE 5.
+
+$ bCompileForSize (Boolean): True if we want to favor optimizing size over speed.
+
+$ bMergeModuleAndGeneratedUnityFiles (Boolean): Whether to merge module and generated unity files for faster compilation.
+
+$ bDetailedUnityFiles (Boolean): Whether to add additional information to the unity files, such as '_of_X' in the file name. Not recommended.
+
+$ bDisableDebugInfo (Boolean): Whether to globally disable debug info generation; Obsolete, please use TargetRules.DebugInfoMode instead
+
+$ bEnableCppCoroutinesForEvaluation (Boolean): Whether to enable support for C++20 coroutines This option is provided to facilitate evaluation of the feature. Expect the name of the option to change. Use of coroutines in with UE is untested and unsupported.
+
+$ bUseMallocProfiler (Boolean): If true, then enable memory profiling in the build (defines USE_MALLOC_PROFILER=1 and forces bOmitFramePointers=false).
+
+$ bDisableInliningGenCpps (Boolean): Disable supports for inlining gen.cpps
+
+$ DisablePluginsConflictWarningLevel (WarningLevel): How to treat conflicts when a disabled plugin is being enabled by another plugin referencing it
+
+$ DeprecationWarningLevel (WarningLevel): Level to report deprecation warnings as errors
+
+$ ShadowVariableWarningLevel (WarningLevel): Forces shadow variable warnings to be treated as errors on platforms that support it.
+
+$ UnsafeTypeCastWarningLevel (WarningLevel): Indicates what warning/error level to treat unsafe type casts as on platforms that support it (e.g., double->float or int64->int32)
+
+$ bUndefinedIdentifierErrors (Boolean): Forces the use of undefined identifiers in conditional expressions to be treated as errors.
+
+$ UndefinedIdentifierWarningLevel (WarningLevel): Indicates what warning/error level to treat undefined identifiers in conditional expressions.
+
+$ PCHPerformanceIssueWarningLevel (WarningLevel): Indicates what warning/error level to treat potential PCH performance issues.
+
+$ ModuleIncludePathWarningLevel (WarningLevel): How to treat general module include path validation messages
+
+$ ModuleIncludePrivateWarningLevel (WarningLevel): How to treat private module include path validation messages, where a module is adding an include path that exposes private headers
+
+$ ModuleIncludeSubdirectoryWarningLevel (WarningLevel): How to treat unnecessary module sub-directory include path validation messages
+
+
