Documentation for build tools

diff --git a/Engine/Documentation/Source/Programming/UnrealBuildSystem/ModuleFiles/ModuleFilesProperties/ModuleFilesProperties.INT.udn b/Engine/Documentation/Source/Programming/UnrealBuildSystem/ModuleFiles/ModuleFilesProperties/ModuleFilesProperties.INT.udn
--- a/Engine/Documentation/Source/Programming/UnrealBuildSystem/ModuleFiles/ModuleFilesProperties/ModuleFilesProperties.INT.udn
+++ b/Engine/Documentation/Source/Programming/UnrealBuildSystem/ModuleFiles/ModuleFilesProperties/ModuleFilesProperties.INT.udn
@@ -48,6 +48,8 @@ $ bAlwaysExportStructs (Boolean): If true, autogenerated functions of USTRUCTS a
 
 $ bAlwaysExportEnums (Boolean): If true, autogenerated functions of UENUMS are always exported
 
+$ bAllowUETypesInNamespaces (Boolean): If true, this module can contain UE types (i.e. UCLASS) in namespaces
+
 $ PublicIncludePathModuleNames (List&lt;String&gt;): List of modules names (no path needed) with header files that our module's public headers needs access to, but we don't need to "import" or link against.
 
 $ PublicDependencyModuleNames (List&lt;String&gt;): List of public dependency module names (no path needed) (automatically does the private/public include). These are modules that are required by our public source files.
