Documentation for build tools

diff --git a/Engine/Documentation/Source/Programming/UnrealBuildSystem/Configuration/BuildConfigProperties/BuildConfigProperties.INT.udn b/Engine/Documentation/Source/Programming/UnrealBuildSystem/Configuration/BuildConfigProperties/BuildConfigProperties.INT.udn
--- a/Engine/Documentation/Source/Programming/UnrealBuildSystem/Configuration/BuildConfigProperties/BuildConfigProperties.INT.udn
+++ b/Engine/Documentation/Source/Programming/UnrealBuildSystem/Configuration/BuildConfigProperties/BuildConfigProperties.INT.udn
@@ -179,6 +179,8 @@ $ ThinLTOCacheDirectory : Directory where to put the ThinLTO cache on supported
 
 $ ThinLTOCachePruningArguments : Arguments that will be applied to prune the ThinLTO cache on supported platforms. Arguments will only be applied if ThinLTOCacheDirectory is set.
 
+$ bThinLTODistributed : When ThinLTO is enabled we use distributed linking if supported This can save a lot of time for large binaries
+
 $ bPGOProfile : Whether to enable Profile Guided Optimization (PGO) instrumentation in this build.
 
 $ bPGOOptimize : Whether to optimize this build with Profile Guided Optimization (PGO).
diff --git a/Engine/Documentation/Source/Programming/UnrealBuildSystem/TargetFiles/TargetFilesProperties/TargetFilesProperties.INT.udn b/Engine/Documentation/Source/Programming/UnrealBuildSystem/TargetFiles/TargetFilesProperties/TargetFilesProperties.INT.udn
--- a/Engine/Documentation/Source/Programming/UnrealBuildSystem/TargetFiles/TargetFilesProperties/TargetFilesProperties.INT.udn
+++ b/Engine/Documentation/Source/Programming/UnrealBuildSystem/TargetFiles/TargetFilesProperties/TargetFilesProperties.INT.udn
@@ -444,6 +444,8 @@ $ ThinLTOCacheDirectory (String): Directory where to put the ThinLTO cache on su
 
 $ ThinLTOCachePruningArguments (String): Arguments that will be applied to prune the ThinLTO cache on supported platforms. Arguments will only be applied if ThinLTOCacheDirectory is set.
 
+$ bThinLTODistributed (Boolean): When ThinLTO is enabled we use distributed linking if supported This can save a lot of time for large binaries
+
 $ bPGOProfile (Boolean): Whether to enable Profile Guided Optimization (PGO) instrumentation in this build.
 
 $ bPGOOptimize (Boolean): Whether to optimize this build with Profile Guided Optimization (PGO).
