Documentation for build tools

diff --git a/Engine/Documentation/Source/Programming/UnrealBuildSystem/Configuration/BuildConfigProperties/BuildConfigProperties.INT.udn b/Engine/Documentation/Source/Programming/UnrealBuildSystem/Configuration/BuildConfigProperties/BuildConfigProperties.INT.udn
--- a/Engine/Documentation/Source/Programming/UnrealBuildSystem/Configuration/BuildConfigProperties/BuildConfigProperties.INT.udn
+++ b/Engine/Documentation/Source/Programming/UnrealBuildSystem/Configuration/BuildConfigProperties/BuildConfigProperties.INT.udn
@@ -147,6 +147,8 @@ $ bDeterministic : Set flags require for deterministic compiling and linking. En
 
 $ bUseVFS : Experimental work in progress: Set flags to use toolchain virtual file system support, to generate consistent pathing in outputs
 
+$ CacheSalt : Salt value to modify cache bucket calculation, can be used to invalidate cached actions outputs
+
 $ bChainPCHs : Whether PCHs should be chained when compiling with clang.
 
 $ bForceIncludePCHHeadersForGenCppFilesWhenPCHIsDisabled : Whether PCH headers should be force included for gen.cpp files when PCH is disabled.
diff --git a/Engine/Documentation/Source/Programming/UnrealBuildSystem/TargetFiles/TargetFilesProperties/TargetFilesProperties.INT.udn b/Engine/Documentation/Source/Programming/UnrealBuildSystem/TargetFiles/TargetFilesProperties/TargetFilesProperties.INT.udn
--- a/Engine/Documentation/Source/Programming/UnrealBuildSystem/TargetFiles/TargetFilesProperties/TargetFilesProperties.INT.udn
+++ b/Engine/Documentation/Source/Programming/UnrealBuildSystem/TargetFiles/TargetFilesProperties/TargetFilesProperties.INT.udn
@@ -406,6 +406,8 @@ $ bDeterministic (Boolean): Set flags require for deterministic compiling and li
 
 $ bUseVFS (Boolean): Experimental work in progress: Set flags to use toolchain virtual file system support, to generate consistent pathing in outputs
 
+$ CacheSalt (String): Salt value to modify cache bucket calculation, can be used to invalidate cached actions outputs
+
 $ bChainPCHs (Boolean): Whether PCHs should be chained when compiling with clang.
 
 $ bForceIncludePCHHeadersForGenCppFilesWhenPCHIsDisabled (Boolean): Whether PCH headers should be force included for gen.cpp files when PCH is disabled.
