Documentation for build tools

diff --git a/Engine/Documentation/Source/Programming/UnrealBuildSystem/Configuration/BuildConfigProperties/BuildConfigProperties.INT.udn b/Engine/Documentation/Source/Programming/UnrealBuildSystem/Configuration/BuildConfigProperties/BuildConfigProperties.INT.udn
--- a/Engine/Documentation/Source/Programming/UnrealBuildSystem/Configuration/BuildConfigProperties/BuildConfigProperties.INT.udn
+++ b/Engine/Documentation/Source/Programming/UnrealBuildSystem/Configuration/BuildConfigProperties/BuildConfigProperties.INT.udn
@@ -261,20 +261,22 @@ $ BaseLogFileName : Specifies the file to use for logging.
 
 $ IWYUBaseLogFileName : Specifies the file to use for logging.
 
-$ bStripSymbols : Whether to strip iOS symbols or not (implied by Shipping config).
+$ bStripSymbols : Whether to strip symbols or not (implied by Shipping config).
 
-$ bUseDSYMFiles : Enables the generation of .dsym files. This can be disabled to enable faster iteration times during development.
+$ bCreateStripFlagFile : Whether to create a ".stripped" file when stripping symbols or not.
 
 $ bSkipClangValidation : Disables clang build verification checks on static libraries
 
 $ bEnableAddressSanitizer : Enables address sanitizer (ASan)
 
-$ bEnableLibFuzzer : Enables LibFuzzer.
-
 $ bEnableThreadSanitizer : Enables thread sanitizer (TSan).
 
 $ bEnableUndefinedBehaviorSanitizer : Enables undefined behavior sanitizer (UBSan)
 
+$ bUseDSYMFiles : Enables the generation of .dsym files. This can be disabled to enable faster iteration times during development.
+
+$ bEnableLibFuzzer : Enables LibFuzzer.
+
 $ bEnableMemorySanitizer : Enables memory sanitizer (MSan)
 
 $ bTuneDebugInfoForLLDB : Turns on tuning of debug info for LLDB
