Documentation for build tools

diff --git a/Engine/Documentation/Source/Programming/UnrealBuildSystem/ModuleFiles/ModuleFilesProperties/ModuleFilesProperties.INT.udn b/Engine/Documentation/Source/Programming/UnrealBuildSystem/ModuleFiles/ModuleFilesProperties/ModuleFilesProperties.INT.udn
--- a/Engine/Documentation/Source/Programming/UnrealBuildSystem/ModuleFiles/ModuleFilesProperties/ModuleFilesProperties.INT.udn
+++ b/Engine/Documentation/Source/Programming/UnrealBuildSystem/ModuleFiles/ModuleFilesProperties/ModuleFilesProperties.INT.udn
@@ -50,6 +50,8 @@ $ bAlwaysExportEnums (Boolean): If true, autogenerated functions of UENUMS are a
 
 $ bAllowUETypesInNamespaces (Boolean): If true, this module can contain UE types (i.e. UCLASS) in namespaces
 
+$ bMinimizeGeneratedIncludes (Boolean): If true, minimizes the amount of header files included in .generated.h files, which may require more headers from CoreUObject to be explicitly included.
+
 $ PublicIncludePathModuleNames (List&lt;String&gt;): List of modules names (no path needed) with header files that our module's public headers needs access to, but we don't need to "import" or link against.
 
 $ PublicDependencyModuleNames (List&lt;String&gt;): List of public dependency module names (no path needed) (automatically does the private/public include). These are modules that are required by our public source files.
