Documentation for build tools

diff --git a/Engine/Documentation/Source/Programming/UnrealBuildSystem/Configuration/BuildConfigProperties/BuildConfigProperties.INT.udn b/Engine/Documentation/Source/Programming/UnrealBuildSystem/Configuration/BuildConfigProperties/BuildConfigProperties.INT.udn
--- a/Engine/Documentation/Source/Programming/UnrealBuildSystem/Configuration/BuildConfigProperties/BuildConfigProperties.INT.udn
+++ b/Engine/Documentation/Source/Programming/UnrealBuildSystem/Configuration/BuildConfigProperties/BuildConfigProperties.INT.udn
@@ -283,6 +283,8 @@ $ bTuneDebugInfoForLLDB : Turns on tuning of debug info for LLDB
 
 $ bDisableDumpSyms : Whether to globally disable calling dump_syms
 
+$ bDisableStripSymbols : Whether to globally skip symbol stripping
+
 $ bEnableInstrumentation : Enables instrumentation.
 
 $ bWriteSolutionOptionFile : Whether to write a solution option (suo) file for the sln.
diff --git a/Engine/Documentation/Source/Programming/UnrealBuildSystem/TargetFiles/TargetFilesProperties/TargetFilesProperties.INT.udn b/Engine/Documentation/Source/Programming/UnrealBuildSystem/TargetFiles/TargetFilesProperties/TargetFilesProperties.INT.udn
--- a/Engine/Documentation/Source/Programming/UnrealBuildSystem/TargetFiles/TargetFilesProperties/TargetFilesProperties.INT.udn
+++ b/Engine/Documentation/Source/Programming/UnrealBuildSystem/TargetFiles/TargetFilesProperties/TargetFilesProperties.INT.udn
@@ -548,7 +548,9 @@ $ bStripExports (Boolean): Experimental: Strip unused exports from libraries. On
 
 $ bMergeModules (Boolean): Experimental: Merge modular modules into combined libraries. Sets LinkType to Modular and enables bStripExports
 
-$ MergeOnlyGameFeatures (Boolean): Experimental: Whether to merge only game feature plugins, ignoring regular plugins.
+$ bMergeOnlyGameFeatures (Boolean): Experimental: Whether to merge only game feature plugins, ignoring regular plugins.
+
+$ bReportModuleMergingData (Boolean): Experimental: Whether to report module merging data to the game target.
 
 $ LaunchModuleName (String): Specifies the name of the launch module. For modular builds, this is the module that is compiled into the target's executable.
 
