Documentation for build tools

diff --git a/Engine/Documentation/Source/Programming/UnrealBuildSystem/Configuration/BuildConfigProperties/BuildConfigProperties.INT.udn b/Engine/Documentation/Source/Programming/UnrealBuildSystem/Configuration/BuildConfigProperties/BuildConfigProperties.INT.udn
--- a/Engine/Documentation/Source/Programming/UnrealBuildSystem/Configuration/BuildConfigProperties/BuildConfigProperties.INT.udn
+++ b/Engine/Documentation/Source/Programming/UnrealBuildSystem/Configuration/BuildConfigProperties/BuildConfigProperties.INT.udn
@@ -135,6 +135,10 @@ $ DebugInfo : How much debug info should be generated. See DebugInfoMode enum fo
 
 $ DebugInfoLineTablesOnly : True if only debug line number tables should be emitted in debug information for compilers that support doing so. Overrides TargetRules.DebugInfo See https://clang.llvm.org/docs/UsersManual.html#cmdoption-gline-tables-only for more information
 
+$ DebugInfoNoInlineLineTables : True if debug line number tables should be omitted in debug information for inline functions, compilers that support doing so.
+
+$ DebugInfoSimpleTemplateNames : True if template type names should be simplified to reduce the size of debug info for compilers that support doing so.
+
 $ bDisableDebugInfoForGeneratedCode : Whether to disable debug info generation for generated files. This improves link times and reduces pdb size for modules that have a lot of generated glue code.
 
 $ bOmitPCDebugInfoInDevelopment : Whether to disable debug info on PC/Mac in development builds (for faster developer iteration, as link times are extremely fast with debug info disabled).
diff --git a/Engine/Documentation/Source/Programming/UnrealBuildSystem/ModuleFiles/ModuleFilesProperties/ModuleFilesProperties.INT.udn b/Engine/Documentation/Source/Programming/UnrealBuildSystem/ModuleFiles/ModuleFilesProperties/ModuleFilesProperties.INT.udn
--- a/Engine/Documentation/Source/Programming/UnrealBuildSystem/ModuleFiles/ModuleFilesProperties/ModuleFilesProperties.INT.udn
+++ b/Engine/Documentation/Source/Programming/UnrealBuildSystem/ModuleFiles/ModuleFilesProperties/ModuleFilesProperties.INT.udn
@@ -102,6 +102,8 @@ $ PrivateDefinitions (List&lt;String&gt;): Private compiler definitions for this
 
 $ PublicDefinitions (List&lt;String&gt;): Public compiler definitions for this module
 
+$ AutoRTFMExternalMappingFiles (List&lt;String&gt;): AutoRTFM external open -> closed mapping files
+
 $ DynamicallyLoadedModuleNames (List&lt;String&gt;): Addition modules this module may require at run-time
 
 $ RuntimeDependencies (RuntimeDependencyList): List of files which this module depends on at runtime. These files will be staged along with the target.
diff --git a/Engine/Documentation/Source/Programming/UnrealBuildSystem/TargetFiles/TargetFilesProperties/TargetFilesProperties.INT.udn b/Engine/Documentation/Source/Programming/UnrealBuildSystem/TargetFiles/TargetFilesProperties/TargetFilesProperties.INT.udn
--- a/Engine/Documentation/Source/Programming/UnrealBuildSystem/TargetFiles/TargetFilesProperties/TargetFilesProperties.INT.udn
+++ b/Engine/Documentation/Source/Programming/UnrealBuildSystem/TargetFiles/TargetFilesProperties/TargetFilesProperties.INT.udn
@@ -26,6 +26,14 @@ $ DebugInfoLineTablesOnlyModules (HashSet&lt;String&gt;): Modules that should em
 
 $ DebugInfoLineTablesOnlyPlugins (HashSet&lt;String&gt;): Plugins that should emit line number tables instead of full debug information for compilers that support doing so. Overrides DisableDebugInfoPlugins
 
+$ DebugInfoNoInlineLineTablesModules (HashSet&lt;String&gt;): Modules that should omit line number tables for inline functions for compilers that support doing so.
+
+$ DebugInfoNoInlineLineTablesPlugins (HashSet&lt;String&gt;): Plugins that should omit line number tables for inline functions for compilers that support doing so.
+
+$ DebugInfoSimpleTemplateNamesModules (HashSet&lt;String&gt;): Modules for which template type names should be simplified to reduce the size of debug info for compilers that support doing so.
+
+$ DebugInfoSimpleTemplateNamesPlugins (HashSet&lt;String&gt;): Plugins for which template type names should be simplified to reduce the size of debug info for compilers that support doing so.
+
 $ MergeAdditionalCommonModules (HashSet&lt;String&gt;): Experimental: List of modules to exclude from Engine to be placed in Common
 
 $ MergePlugins (List&lt;String&gt;): Experimental: List of plugins (and their dependencies) to each merge into separate libraries. Requires bMergeModules to be enabled
@@ -394,6 +402,10 @@ $ DebugInfo (DebugInfoMode): How much debug info should be generated. See DebugI
 
 $ DebugInfoLineTablesOnly (DebugInfoMode): True if only debug line number tables should be emitted in debug information for compilers that support doing so. Overrides TargetRules.DebugInfo See https://clang.llvm.org/docs/UsersManual.html#cmdoption-gline-tables-only for more information
 
+$ DebugInfoNoInlineLineTables (DebugInfoMode): True if debug line number tables should be omitted in debug information for inline functions, compilers that support doing so.
+
+$ DebugInfoSimpleTemplateNames (DebugInfoMode): True if template type names should be simplified to reduce the size of debug info for compilers that support doing so.
+
 $ bDisableDebugInfoForGeneratedCode (Boolean): Whether to disable debug info generation for generated files. This improves link times and reduces pdb size for modules that have a lot of generated glue code.
 
 $ bOmitPCDebugInfoInDevelopment (Boolean): Whether to disable debug info on PC/Mac in development builds (for faster developer iteration, as link times are extremely fast with debug info disabled).
