Documentation for build tools

diff --git a/Engine/Documentation/Source/Programming/UnrealBuildSystem/TargetFiles/TargetFilesProperties/TargetFilesProperties.INT.udn b/Engine/Documentation/Source/Programming/UnrealBuildSystem/TargetFiles/TargetFilesProperties/TargetFilesProperties.INT.udn
--- a/Engine/Documentation/Source/Programming/UnrealBuildSystem/TargetFiles/TargetFilesProperties/TargetFilesProperties.INT.udn
+++ b/Engine/Documentation/Source/Programming/UnrealBuildSystem/TargetFiles/TargetFilesProperties/TargetFilesProperties.INT.udn
@@ -314,6 +314,8 @@ $ bUseExecCommandsInShipping (Boolean): Whether to turn allow exec commands for
 
 $ OptimizationLevel (OptimizationMode): Allows to fine tune optimizations level for speed and\or code size
 
+$ bDefineForceInlineHintToInline (Boolean): By default UE_FORCEINLINE_HINT is defined to FORCEINLINE. If this flag is true it will instead be defined to 'inline'. This flag should only be true for targets where code size is a concern and/or is being compiled with PGO or LTO. Note that if this flag is true when PGOOptimize is enabled, it should also be true when PGOProfile is enabled.
+
 $ FPSemantics (FPSemanticsMode): Allows setting the FP semantics.
 
 $ bForceCompileDevelopmentAutomationTests (Boolean): Whether to compile development automation tests.
