Documentation for build tools

diff --git a/Engine/Documentation/Source/Programming/UnrealBuildSystem/Configuration/BuildConfigProperties/BuildConfigProperties.INT.udn b/Engine/Documentation/Source/Programming/UnrealBuildSystem/Configuration/BuildConfigProperties/BuildConfigProperties.INT.udn
--- a/Engine/Documentation/Source/Programming/UnrealBuildSystem/Configuration/BuildConfigProperties/BuildConfigProperties.INT.udn
+++ b/Engine/Documentation/Source/Programming/UnrealBuildSystem/Configuration/BuildConfigProperties/BuildConfigProperties.INT.udn
@@ -611,6 +611,10 @@ $ LocalHost : Which ip UBA server should give to agents. This will invert so hos
 
 $ MaxCores : Max cores allowed to be used by build session
 
+$ MaxWorkers : Max workers allowed to be used by build session
+
+$ MaxIdle : Max idle time in seconds to maintain a connection to a worker when no pending actions are currently available due to waiting on dependencies
+
 $ StartupDelay : How long UBT should wait to ask for help. Useful in build configs where machine can delay remote work and still get same wall time results (pch dependencies etc)
 
 $ AllowWine : Allow use of Wine. Only applicable to Horde agents running Linux. Can still be ignored if Wine executable is not set on agent.
