Documentation for build tools

diff --git a/Engine/Documentation/Source/Programming/UnrealBuildSystem/TargetFiles/TargetFilesProperties/TargetFilesProperties.INT.udn b/Engine/Documentation/Source/Programming/UnrealBuildSystem/TargetFiles/TargetFilesProperties/TargetFilesProperties.INT.udn
--- a/Engine/Documentation/Source/Programming/UnrealBuildSystem/TargetFiles/TargetFilesProperties/TargetFilesProperties.INT.udn
+++ b/Engine/Documentation/Source/Programming/UnrealBuildSystem/TargetFiles/TargetFilesProperties/TargetFilesProperties.INT.udn
@@ -26,9 +26,9 @@ $ DebugInfoLineTablesOnlyModules (HashSet&lt;String&gt;): Modules that should em
 
 $ DebugInfoLineTablesOnlyPlugins (HashSet&lt;String&gt;): Plugins that should emit line number tables instead of full debug information for compilers that support doing so. Overrides DisableDebugInfoPlugins
 
-$ DebugInfoNoInlineLineTablesModules (HashSet&lt;String&gt;): Modules that should omit line number tables for inline functions for compilers that support doing so.
+$ DebugInfoInlineLineTablesModules (HashSet&lt;String&gt;): Modules that should include line number tables for inline functions for compilers that support doing so, even if DebugInfoNoInlineLineTables would otherwise remove them.
 
-$ DebugInfoNoInlineLineTablesPlugins (HashSet&lt;String&gt;): Plugins that should omit line number tables for inline functions for compilers that support doing so.
+$ DebugInfoInlineLineTablesPlugins (HashSet&lt;String&gt;): Modules that should include line number tables for inline functions for compilers that support doing so, even if DebugInfoNoInlineLineTables would otherwise remove them.
 
 $ DebugInfoSimpleTemplateNamesModules (HashSet&lt;String&gt;): Modules for which template type names should be simplified to reduce the size of debug info for compilers that support doing so.
 
