Documentation for build tools

diff --git a/Engine/Documentation/Source/Programming/UnrealBuildSystem/Configuration/BuildConfigProperties/BuildConfigProperties.INT.udn b/Engine/Documentation/Source/Programming/UnrealBuildSystem/Configuration/BuildConfigProperties/BuildConfigProperties.INT.udn
--- a/Engine/Documentation/Source/Programming/UnrealBuildSystem/Configuration/BuildConfigProperties/BuildConfigProperties.INT.udn
+++ b/Engine/Documentation/Source/Programming/UnrealBuildSystem/Configuration/BuildConfigProperties/BuildConfigProperties.INT.udn
@@ -449,6 +449,10 @@ $ WriteCache : Set cache to write, expects one of [True, False, BuildMachineOnly
 
 $ RequireVFS : Set cache to require VFS to be enabled
 
+$ CacheCrypto : Crypto used to connect to cache server.
+
+$ CacheDesiredConnectionCount : Desired connection count. When there is latency it might be a good idea to have multiple tcp connections
+
 $ Providers : List of provider names for remote connections.
 
 $ BuildMachineProviders : List of provider names for remote connections
@@ -525,6 +529,8 @@ $ CacheShuffle : Shuffle the list of cache servers before each read attempt
 
 $ ReportCacheMissReason : Report reason a cache miss happened. Useful when searching for determinism/portability issues
 
+$ CacheLinkActions : Control if link actions should be downloaded from cache or not
+
 $ CompressionLevel : Compression level. Options are None, SuperFast, VeryFast, Fast, Normal, Optimal1 to Optimal5
 
 $ bDisableHorde : Disable horde all together
