Documentation for build tools

diff --git a/Engine/Documentation/Source/Programming/UnrealBuildSystem/Configuration/BuildConfigProperties/BuildConfigProperties.INT.udn b/Engine/Documentation/Source/Programming/UnrealBuildSystem/Configuration/BuildConfigProperties/BuildConfigProperties.INT.udn
--- a/Engine/Documentation/Source/Programming/UnrealBuildSystem/Configuration/BuildConfigProperties/BuildConfigProperties.INT.udn
+++ b/Engine/Documentation/Source/Programming/UnrealBuildSystem/Configuration/BuildConfigProperties/BuildConfigProperties.INT.udn
@@ -59,6 +59,8 @@ $ UndefinedIdentifierWarningLevel : Indicates what warning/error level to treat
 
 $ SwitchUnhandledEnumeratorWarningLevel : Indicates what warning/error level to treat unhandled enumerators in switches on enumeration-typed values. MSVC - 4061 - https://learn.microsoft.com/en-us/cpp/error-messages/compiler-warnings/compiler-warning-level-4-c4061 44061 - Note: The extra 4 is not a typo, /wLXXXX sets warning XXXX to level L Clang - https://clang.llvm.org/docs/DiagnosticsReference.html#wswitch-enum
 
+$ ShortenSizeTToIntWarningLevel : The compiler detected a conversion from size_t to a smaller type. MSVC - https://learn.microsoft.com/en-us/cpp/error-messages/compiler-warnings/compiler-warning-level-3-c4267 4267 - 'var' : conversion from 'size_t' to 'type', possible loss of data
+
 $ DeprecationWarningLevel : Indicates what warning/error level to treat usages of functions, methods, or variables, that are marked as deprecated. Clang - https://clang.llvm.org/docs/DiagnosticsReference.html#wdeprecated-declarations
 
 $ PCHPerformanceIssueWarningLevel : Indicates what warning/error level to treat potential PCH performance issues.
