Documentation for build tools

diff --git a/Engine/Documentation/Source/Programming/UnrealBuildSystem/ModuleFiles/ModuleFilesProperties/ModuleFilesProperties.INT.udn b/Engine/Documentation/Source/Programming/UnrealBuildSystem/ModuleFiles/ModuleFilesProperties/ModuleFilesProperties.INT.udn
--- a/Engine/Documentation/Source/Programming/UnrealBuildSystem/ModuleFiles/ModuleFilesProperties/ModuleFilesProperties.INT.udn
+++ b/Engine/Documentation/Source/Programming/UnrealBuildSystem/ModuleFiles/ModuleFilesProperties/ModuleFilesProperties.INT.udn
@@ -38,6 +38,8 @@ $ IsVcPackageSupported (Boolean): Returns if VcPkg is supported for the build co
 
 
 ### Read/Write Properties
+$ FilesToGenerate (Dictionary&lt;String, IEnumerable&lt;String&gt;&gt;): List of files that will be generated for this module. Entries in list are (FileName, FileContent)
+
 $ StaticAnalyzerRulesets (HashSet&lt;FileReference&gt;): The static analyzer rulesets that should be used to filter warnings. This is only supported for MSVC. See https://learn.microsoft.com/en-us/cpp/code-quality/using-rule-sets-to-specify-the-cpp-rules-to-run
 
 $ AllowedRestrictedFolders (List&lt;String&gt;): List of folders which are allowed to be referenced when compiling this binary, without propagating restricted folder names
diff --git a/Engine/Documentation/Source/Programming/UnrealBuildSystem/TargetFiles/TargetFilesProperties/TargetFilesProperties.INT.udn b/Engine/Documentation/Source/Programming/UnrealBuildSystem/TargetFiles/TargetFilesProperties/TargetFilesProperties.INT.udn
--- a/Engine/Documentation/Source/Programming/UnrealBuildSystem/TargetFiles/TargetFilesProperties/TargetFilesProperties.INT.udn
+++ b/Engine/Documentation/Source/Programming/UnrealBuildSystem/TargetFiles/TargetFilesProperties/TargetFilesProperties.INT.udn
@@ -50,6 +50,8 @@ $ GlobalDefinitions (List&lt;String&gt;): Macros to define globally across the w
 
 $ ProjectDefinitions (List&lt;String&gt;): Macros to define across all macros in the project.
 
+$ AdditionalProperties (AdditionalProperties): Environment variables. Should be used instead of Environment.Set/GetEnvironmentVariable in .build.cs
+
 $ ExtraModuleNames (List&lt;String&gt;): List of additional modules to be compiled into the target.
 
 $ ManifestFileNames (List&lt;FileReference&gt;): Path to a manifest to output for this target
