Documentation for build tools

diff --git a/Engine/Documentation/Source/Programming/UnrealBuildSystem/Configuration/BuildConfigProperties/BuildConfigProperties.INT.udn b/Engine/Documentation/Source/Programming/UnrealBuildSystem/Configuration/BuildConfigProperties/BuildConfigProperties.INT.udn
--- a/Engine/Documentation/Source/Programming/UnrealBuildSystem/Configuration/BuildConfigProperties/BuildConfigProperties.INT.udn
+++ b/Engine/Documentation/Source/Programming/UnrealBuildSystem/Configuration/BuildConfigProperties/BuildConfigProperties.INT.udn
@@ -351,7 +351,7 @@ $ bVCFastFail : True if /fastfail should be passed to the msvc compiler and link
 
 $ bVCExtendedWarningInfo : True if /d2ExtendedWarningInfo should be passed to the compiler and /d2:-ExtendedWarningInfo to the linker
 
-$ bClangStandaloneDebug : True if optimizations to reduce the size of debug information should be disabled See https://clang.llvm.org/docs/UsersManual.html#cmdoption-fstandalone-debug for more information
+$ bClangStandaloneDebug : True if optimizations to reduce the size of debug information should be disabled. Disabling this option can decrease pdb size but may degrade the debugging experience. See https://clang.llvm.org/docs/UsersManual.html#cmdoption-fstandalone-debug for more information
 
 $ bAllowClangLinker : True if we should use the Clang linker (LLD) when we are compiling with Clang or Intel oneAPI, otherwise we use the MSVC linker.
 
