Documentation for build tools

diff --git a/Engine/Documentation/Source/Programming/UnrealBuildSystem/ModuleFiles/ModuleFilesProperties/ModuleFilesProperties.INT.udn b/Engine/Documentation/Source/Programming/UnrealBuildSystem/ModuleFiles/ModuleFilesProperties/ModuleFilesProperties.INT.udn
--- a/Engine/Documentation/Source/Programming/UnrealBuildSystem/ModuleFiles/ModuleFilesProperties/ModuleFilesProperties.INT.udn
+++ b/Engine/Documentation/Source/Programming/UnrealBuildSystem/ModuleFiles/ModuleFilesProperties/ModuleFilesProperties.INT.udn
@@ -212,6 +212,8 @@ $ bAllowConfidentialPlatformDefines (Boolean): Whether this module can use PLATF
 
 $ bDisableAutoRTFMInstrumentation (Boolean): Disables AutoRTFM instrumentation to this module only when AutoRTFMCompiler is enabled
 
+$ bInjectAutoRTFMAttributeDisables (Boolean): If `bDisableAutoRTFMInstrumentation` is true, will make the modules `*_API` define contain `AUTORTFM_DISABLE` to let us enable errors for code that uses transactionally-disabled code in places that are transactionally-enabled.
+
 $ PrecompileForTargets (PrecompileTargetsType): Which targets this module should be precompiled for
 
 $ bRequiresImplementModule (Nullable&lt;Boolean&gt;): Whether this module requires the IMPLEMENT_MODULE macro to be implemented. Most UE modules require this, since we use the IMPLEMENT_MODULE macro to do other global overloads (eg. operator new/delete forwarding to GMalloc).
