Documentation for build tools

diff --git a/Engine/Documentation/Source/Programming/Development/BuildGraph/ScriptAnatomy/Tasks/PredefinedTasks/BuildGraphScriptPredefinedTasks.INT.udn b/Engine/Documentation/Source/Programming/Development/BuildGraph/ScriptAnatomy/Tasks/PredefinedTasks/BuildGraphScriptPredefinedTasks.INT.udn
--- a/Engine/Documentation/Source/Programming/Development/BuildGraph/ScriptAnatomy/Tasks/PredefinedTasks/BuildGraphScriptPredefinedTasks.INT.udn
+++ b/Engine/Documentation/Source/Programming/Development/BuildGraph/ScriptAnatomy/Tasks/PredefinedTasks/BuildGraphScriptPredefinedTasks.INT.udn
@@ -169,7 +169,8 @@ Uploads an artifact to Cloud DDC
 |                          |               |          |                                                                                                                                                                                                                          |
 | ------------------------ | ------------- | -------- | ------------------------------------------------------------------------------------------------------------------------------------------------------------------------------------------------------------------------ |
 | Name                     | String        | Required | Name of the artifact                                                                                                                                                                                                     |
-| HordeArtifactName        | String        | Optional | Optional override of the name of the artifact as presented in Horde, defaults to name if not set                                                                                                                         |
+| HordeArtifactID          | String        | Optional | An optional override of the id of the artifact as presented in Horde, defaults to name if not set                                                                                                                        |
+| BuildGroup               | String        | Required | An identifier to set for a group of builds that are compatible                                                                                                                                                           |
 | Type                     | String        | Required | The artifact type. Determines the permissions and expiration policy for the artifact.                                                                                                                                    |
 | Description              | String        | Optional | Description for the artifact. Will be shown through the Horde dashboard.                                                                                                                                                 |
 | Platform                 | String        | Required | The platform the artifact is created for                                                                                                                                                                                 |
