Documentation for build tools

diff --git a/Engine/Documentation/Source/Programming/UnrealBuildSystem/Configuration/BuildConfigProperties/BuildConfigProperties.INT.udn b/Engine/Documentation/Source/Programming/UnrealBuildSystem/Configuration/BuildConfigProperties/BuildConfigProperties.INT.udn
--- a/Engine/Documentation/Source/Programming/UnrealBuildSystem/Configuration/BuildConfigProperties/BuildConfigProperties.INT.udn
+++ b/Engine/Documentation/Source/Programming/UnrealBuildSystem/Configuration/BuildConfigProperties/BuildConfigProperties.INT.udn
@@ -179,7 +179,7 @@ $ bForcePrecompiledHeaderForGameModules : When enabled, a precompiled header is
 
 $ bUseIncrementalLinking : Whether to use incremental linking or not. Incremental linking can yield faster iteration times when making small changes. Currently disabled by default because it tends to behave a bit buggy on some computers (PDB-related compile errors).
 
-$ bAllowLTCG : Whether to allow the use of link time code generation (LTCG).
+$ bAllowLTCG : Whether to allow the use of link time code generation (LTCG). Implied when PGO is enabled
 
 $ bPreferThinLTO : When Link Time Code Generation (LTCG) is enabled, whether to prefer using the lighter weight version on supported platforms.
 
diff --git a/Engine/Documentation/Source/Programming/UnrealBuildSystem/TargetFiles/TargetFilesProperties/TargetFilesProperties.INT.udn b/Engine/Documentation/Source/Programming/UnrealBuildSystem/TargetFiles/TargetFilesProperties/TargetFilesProperties.INT.udn
--- a/Engine/Documentation/Source/Programming/UnrealBuildSystem/TargetFiles/TargetFilesProperties/TargetFilesProperties.INT.udn
+++ b/Engine/Documentation/Source/Programming/UnrealBuildSystem/TargetFiles/TargetFilesProperties/TargetFilesProperties.INT.udn
@@ -454,7 +454,7 @@ $ bForcePrecompiledHeaderForGameModules (Boolean): When enabled, a precompiled h
 
 $ bUseIncrementalLinking (Boolean): Whether to use incremental linking or not. Incremental linking can yield faster iteration times when making small changes. Currently disabled by default because it tends to behave a bit buggy on some computers (PDB-related compile errors).
 
-$ bAllowLTCG (Boolean): Whether to allow the use of link time code generation (LTCG).
+$ bAllowLTCG (Boolean): Whether to allow the use of link time code generation (LTCG). Implied when PGO is enabled
 
 $ bPreferThinLTO (Boolean): When Link Time Code Generation (LTCG) is enabled, whether to prefer using the lighter weight version on supported platforms.
 
