Documentation for build tools

diff --git a/Engine/Documentation/Source/Programming/UnrealBuildSystem/TargetFiles/TargetFilesProperties/TargetFilesProperties.INT.udn b/Engine/Documentation/Source/Programming/UnrealBuildSystem/TargetFiles/TargetFilesProperties/TargetFilesProperties.INT.udn
--- a/Engine/Documentation/Source/Programming/UnrealBuildSystem/TargetFiles/TargetFilesProperties/TargetFilesProperties.INT.udn
+++ b/Engine/Documentation/Source/Programming/UnrealBuildSystem/TargetFiles/TargetFilesProperties/TargetFilesProperties.INT.udn
@@ -208,8 +208,6 @@ $ bUseGameplayDebuggerCore (Boolean): Set to true when bUseGameplayDebugger is f
 
 $ bCompileIoStoreOnDemand (Boolean): Whether to use I/O store on-demand
 
-$ bUseIris (Boolean): Whether to use Iris.
-
 $ bTrackRHIResourceInfoForTest (Boolean): Whether to track owner (asset name) of RHI resource for Test configuration. Useful for ListShaderMaps and ListShaderLibraries commands.
 
 $ bBuildEditor (Boolean): Whether we are compiling editor code or not. Prefer the more explicit bCompileAgainstEditor instead.
