Documentation for build tools

diff --git a/Engine/Documentation/Source/Programming/UnrealBuildSystem/Configuration/BuildConfigProperties/BuildConfigProperties.INT.udn b/Engine/Documentation/Source/Programming/UnrealBuildSystem/Configuration/BuildConfigProperties/BuildConfigProperties.INT.udn
--- a/Engine/Documentation/Source/Programming/UnrealBuildSystem/Configuration/BuildConfigProperties/BuildConfigProperties.INT.udn
+++ b/Engine/Documentation/Source/Programming/UnrealBuildSystem/Configuration/BuildConfigProperties/BuildConfigProperties.INT.udn
@@ -715,6 +715,34 @@ $ bUseDangerouslyFastMode : If this is set, then we do not do any post-compile s
 
 $ ObjSrcMapFile : Whether we should export a file containing .obj to source file mappings.
 
+### AndroidStudioFileGenerator
+
+$ bIncludeDocumentation : True if we should include documentation in the generated projects.
+
+$ bUsePrecompiled : True if build targets should pass the -useprecompiled argument.
+
+$ bIncludeEngineSource : True if we should include engine source in the generated solution.
+
+$ bIncludeShaderSource : True if shader source files should be included in generated projects.
+
+$ bIncludeConfigFiles : True if we should include config (.ini) files in the generated project.
+
+$ bIncludeTemplateFiles : True if we should include template files in the generated project.
+
+$ bIncludeEnginePrograms : True if we should include program projects in the generated solution.
+
+$ IncludeCppSource : Whether to include C++ targets
+
+$ bIncludeDotNetPrograms : True if we should include csharp program projects in the generated solution. Pass "-DotNet" to enable this.
+
+$ bIncludeTempTargets : Whether to include temporary targets generated by UAT to support content only projects with non-default settings.
+
+$ PrimaryProjectName : Name of the primary project file -- for example, the base file name for the Visual Studio solution file, or the Xcode project file on Mac.
+
+$ bPrimaryProjectNameFromFolder : If true, sets the primary project name according to the name of the folder it is in.
+
+$ bPrimaryProjectNameAppendDriveLetter : If true, append the drive letter to the primary project name.
+
 ### CLionGenerator
 
 $ bIncludeDocumentation : True if we should include documentation in the generated projects.
