Documentation for build tools

diff --git a/Engine/Documentation/Source/Programming/UnrealBuildSystem/Configuration/BuildConfigProperties/BuildConfigProperties.INT.udn b/Engine/Documentation/Source/Programming/UnrealBuildSystem/Configuration/BuildConfigProperties/BuildConfigProperties.INT.udn
--- a/Engine/Documentation/Source/Programming/UnrealBuildSystem/Configuration/BuildConfigProperties/BuildConfigProperties.INT.udn
+++ b/Engine/Documentation/Source/Programming/UnrealBuildSystem/Configuration/BuildConfigProperties/BuildConfigProperties.INT.udn
@@ -95,6 +95,8 @@ $ bAutoRTFMClosedStaticLinkage : Whether to link closed function declarations st
 
 $ bUseInlining : Enable inlining for all modules.
 
+$ bWithLiveCoding : Whether to enable support for live coding
+
 $ bUseDebugLiveCodingConsole : Whether to enable support for live coding
 
 $ bUseXGEController : Whether the XGE controller worker and modules should be included in the engine build. These are required for distributed shader compilation using the XGE interception interface.
