Documentation for build tools

diff --git a/Engine/Documentation/Source/Programming/UnrealBuildSystem/Configuration/BuildConfigProperties/BuildConfigProperties.INT.udn b/Engine/Documentation/Source/Programming/UnrealBuildSystem/Configuration/BuildConfigProperties/BuildConfigProperties.INT.udn
--- a/Engine/Documentation/Source/Programming/UnrealBuildSystem/Configuration/BuildConfigProperties/BuildConfigProperties.INT.udn
+++ b/Engine/Documentation/Source/Programming/UnrealBuildSystem/Configuration/BuildConfigProperties/BuildConfigProperties.INT.udn
@@ -257,6 +257,8 @@ $ ActionStallReportTime : Number of second of no completed actions to trigger an
 
 $ ActionStallTerminateTime : Number of second of no completed actions to trigger to terminate the queue. If zero, force termination not be enabled.
 
+$ bPrioritizeDependencyScanning : When enabled, dependency scanning is prioritized to avoid being slowed down by running local jobs.
+
 $ bStopSNDBSCompilationAfterErrors : When enabled, SN-DBS will stop compiling targets after a compile error occurs.  Recommended, as it saves computing resources for others.
 
 $ bXGENoWatchdogThread : Whether to use the no_watchdog_thread option to prevent VS2015 toolchain stalls.
diff --git a/Engine/Documentation/Source/Programming/UnrealBuildSystem/TargetFiles/TargetFilesProperties/TargetFilesProperties.INT.udn b/Engine/Documentation/Source/Programming/UnrealBuildSystem/TargetFiles/TargetFilesProperties/TargetFilesProperties.INT.udn
--- a/Engine/Documentation/Source/Programming/UnrealBuildSystem/TargetFiles/TargetFilesProperties/TargetFilesProperties.INT.udn
+++ b/Engine/Documentation/Source/Programming/UnrealBuildSystem/TargetFiles/TargetFilesProperties/TargetFilesProperties.INT.udn
@@ -236,6 +236,8 @@ $ bCompileAgainstEngine (Boolean): Enabled for all builds that include the engin
 
 $ bCompileAgainstCoreUObject (Boolean): Enabled for all builds that include the CoreUObject project.  Disabled only when building standalone apps that only link with Core.
 
+$ bEnableConstInitUObject (Boolean): Enable compilation of UClass and related reflection objects as statically initialized global variables.
+
 $ bCompileAgainstApplicationCore (Boolean): Enabled for builds that need to initialize the ApplicationCore module. Command line utilities do not normally need this.
 
 $ bEnableTrace (Boolean): Enabled for builds that want to use against the Trace module for profiling and diagnostics.
@@ -562,6 +564,8 @@ $ bStripExports (Boolean): Experimental: Strip unused exports from libraries. On
 
 $ bMergeModules (Boolean): Experimental: Merge modular modules into combined libraries. Sets LinkType to Modular and enables bStripExports
 
+$ bCheckDelegateLifetime (Boolean): Experimental: Check lifetimes for delegate instances. Recommended for merged modules testing
+
 $ bMergeOnlyGameFeatures (Boolean): Experimental: Whether to merge only game feature plugins, ignoring regular plugins.
 
 $ bReportModuleMergingData (Boolean): Experimental: Whether to report module merging data to the game target.
