Documentation for build tools

diff --git a/Engine/Documentation/Source/Programming/UnrealBuildSystem/Configuration/BuildConfigProperties/BuildConfigProperties.INT.udn b/Engine/Documentation/Source/Programming/UnrealBuildSystem/Configuration/BuildConfigProperties/BuildConfigProperties.INT.udn
--- a/Engine/Documentation/Source/Programming/UnrealBuildSystem/Configuration/BuildConfigProperties/BuildConfigProperties.INT.udn
+++ b/Engine/Documentation/Source/Programming/UnrealBuildSystem/Configuration/BuildConfigProperties/BuildConfigProperties.INT.udn
@@ -57,6 +57,8 @@ $ UnsafeTypeCastWarningLevel : Indicates what warning/error level to treat unsaf
 
 $ UndefinedIdentifierWarningLevel : Indicates what warning/error level to treat undefined identifiers in conditional expressions. MSVC - https://learn.microsoft.com/en-us/cpp/error-messages/compiler-warnings/compiler-warning-level-4-c4668 4668 - 'symbol' is not defined as a preprocessor macro, replacing with '0' for 'directives' 44668 - Note: The extra 4 is not a typo, /wLXXXX sets warning XXXX to level L Clang - https://clang.llvm.org/docs/DiagnosticsReference.html#wundef
 
+$ UnreachableCodeWarningLevel : Indicates what warning/error level to treat unreachable code. MSVC - https://learn.microsoft.com/en-us/cpp/error-messages/compiler-warnings/compiler-warning-level-4-c4702 4702 - unreachable code 44702 - Note: The extra 4 is not a typo, /wLXXXX sets warning XXXX to level L Clang - https://clang.llvm.org/docs/DiagnosticsReference.html#wunreachable-code
+
 $ SwitchUnhandledEnumeratorWarningLevel : Indicates what warning/error level to treat unhandled enumerators in switches on enumeration-typed values. MSVC - 4061 - https://learn.microsoft.com/en-us/cpp/error-messages/compiler-warnings/compiler-warning-level-4-c4061 44061 - Note: The extra 4 is not a typo, /wLXXXX sets warning XXXX to level L Clang - https://clang.llvm.org/docs/DiagnosticsReference.html#wswitch-enum
 
 $ ShortenSizeTToIntWarningLevel : The compiler detected a conversion from size_t to a smaller type. MSVC - https://learn.microsoft.com/en-us/cpp/error-messages/compiler-warnings/compiler-warning-level-3-c4267 4267 - 'var' : conversion from 'size_t' to 'type', possible loss of data
