Documentation for build tools

diff --git a/Engine/Documentation/Source/Programming/UnrealBuildSystem/Configuration/BuildConfigProperties/BuildConfigProperties.INT.udn b/Engine/Documentation/Source/Programming/UnrealBuildSystem/Configuration/BuildConfigProperties/BuildConfigProperties.INT.udn
--- a/Engine/Documentation/Source/Programming/UnrealBuildSystem/Configuration/BuildConfigProperties/BuildConfigProperties.INT.udn
+++ b/Engine/Documentation/Source/Programming/UnrealBuildSystem/Configuration/BuildConfigProperties/BuildConfigProperties.INT.udn
@@ -213,7 +213,7 @@ $ bCheckLicenseViolations : Tells the UBT to check if module currently being bui
 
 $ bBreakBuildOnLicenseViolation : Tells the UBT to break build if module currently being built is violating EULA.
 
-$ bUseFastPDBLinking : Whether to use the :FASTLINK option when building with /DEBUG to create local PDBs on Windows. Fast, but currently seems to have problems finding symbols in the debugger.
+$ bUseFastPDBLinking : Whether to use the :FASTLINK option when building with /DEBUG to create local PDBs on Windows. No longer recommended.
 
 $ bCreateMapFile : Outputs a map file as part of the build.
 
@@ -365,16 +365,6 @@ $ WindowsSdkVersion : The specific Windows SDK version to use. This may be a spe
 
 $ bWriteSarif : Whether .sarif files containing errors and warnings are written alongside each .obj, if supported
 
-$ bUpdatedCPPMacro : Enables updated __cplusplus macro (/Zc:__cplusplus).
-
-$ bStrictInlineConformance : Enables inline conformance (Remove unreferenced COMDAT) (/Zc:inline).
-
-$ bStrictPreprocessorConformance : Enables new preprocessor conformance (/Zc:preprocessor). This is always enabled for C++20 modules.
-
-$ bStrictEnumTypesConformance : Enables enum types conformance (/Zc:enumTypes) in VS2022 17.4 Preview 4.0+.
-
-$ bStrictODRViolationConformance : Enables enforcing standard C++ ODR violations (/Zc:checkGwOdr) in VS2022 17.5 Preview 2.0+
-
 $ bStripPrivateSymbols : Whether to request the linker create a stripped pdb file as part of the build. If enabled the full debug pdb will have the extension .full.pdb
 
 $ bNoLinkerDebugInfo : If you supply -NoDebugInfo, windows platforms still create debug info while linking. Set this to true to not create debug info while linking in this circumstance
@@ -409,7 +399,17 @@ $ ToolchainVersionWarningLevel : Warning level when reporting toolchains that ar
 
 $ bStrictConformanceMode : Enables strict standard conformance mode (/permissive-).
 
-$ bDisableVolatileMetadata : Volatile Metadata is enabled by default and improves x64 emulation on arm64, but may come at a small perfomance cost (/volatileMetadata-).
+$ bUpdatedCPPMacro : Enables updated __cplusplus macro (/Zc:__cplusplus).
+
+$ bStrictInlineConformance : Enables inline conformance (Remove unreferenced COMDAT) (/Zc:inline).
+
+$ bStrictPreprocessorConformance : Enables preprocessor conformance (/Zc:preprocessor).
+
+$ bStrictEnumTypesConformance : Enables enum types conformance (/Zc:enumTypes).
+
+$ bStrictODRViolationConformance : Enables enforcing standard C++ ODR violations (/Zc:checkGwOdr).
+
+$ bDisableVolatileMetadata : Volatile Metadata is enabled by default and improves x64 emulation on arm64, but may come at a small performance cost (/volatileMetadata-).
 
 $ StructMemberAlignment : Controls how the members of a structure are packed into memory and specifies the same packing for all structures in a module. See https://learn.microsoft.com/en-us/cpp/build/reference/zp-struct-member-alignment
 
diff --git a/Engine/Documentation/Source/Programming/UnrealBuildSystem/TargetFiles/TargetFilesProperties/TargetFilesProperties.INT.udn b/Engine/Documentation/Source/Programming/UnrealBuildSystem/TargetFiles/TargetFilesProperties/TargetFilesProperties.INT.udn
--- a/Engine/Documentation/Source/Programming/UnrealBuildSystem/TargetFiles/TargetFilesProperties/TargetFilesProperties.INT.udn
+++ b/Engine/Documentation/Source/Programming/UnrealBuildSystem/TargetFiles/TargetFilesProperties/TargetFilesProperties.INT.udn
@@ -492,7 +492,7 @@ $ bCheckLicenseViolations (Boolean): Tells the UBT to check if module currently
 
 $ bBreakBuildOnLicenseViolation (Boolean): Tells the UBT to break build if module currently being built is violating EULA.
 
-$ bUseFastPDBLinking (Nullable&lt;Boolean&gt;): Whether to use the :FASTLINK option when building with /DEBUG to create local PDBs on Windows. Fast, but currently seems to have problems finding symbols in the debugger.
+$ bUseFastPDBLinking (Nullable&lt;Boolean&gt;): Whether to use the :FASTLINK option when building with /DEBUG to create local PDBs on Windows. No longer recommended.
 
 $ bCreateMapFile (Boolean): Outputs a map file as part of the build.
 
