Documentation for build tools

diff --git a/Engine/Documentation/Source/Programming/UnrealBuildSystem/Configuration/BuildConfigProperties/BuildConfigProperties.INT.udn b/Engine/Documentation/Source/Programming/UnrealBuildSystem/Configuration/BuildConfigProperties/BuildConfigProperties.INT.udn
--- a/Engine/Documentation/Source/Programming/UnrealBuildSystem/Configuration/BuildConfigProperties/BuildConfigProperties.INT.udn
+++ b/Engine/Documentation/Source/Programming/UnrealBuildSystem/Configuration/BuildConfigProperties/BuildConfigProperties.INT.udn
@@ -161,8 +161,6 @@ $ bChainPCHs : Whether PCHs should be chained when compiling with clang.
 
 $ bForceIncludePCHHeadersForGenCppFilesWhenPCHIsDisabled : Whether PCH headers should be force included for gen.cpp files when PCH is disabled.
 
-$ bPreprocessDepends : Generate dependency files by preprocessing. This is only recommended when distributing builds as it adds additional overhead.
-
 $ StaticAnalyzer : Whether static code analysis should be enabled.
 
 $ StaticAnalyzerOutputType : The output type to use for the static analyzer. This is only supported for Clang.
@@ -213,8 +211,6 @@ $ bCheckLicenseViolations : Tells the UBT to check if module currently being bui
 
 $ bBreakBuildOnLicenseViolation : Tells the UBT to break build if module currently being built is violating EULA.
 
-$ bUseFastPDBLinking : Whether to use the :FASTLINK option when building with /DEBUG to create local PDBs on Windows. No longer recommended.
-
 $ bCreateMapFile : Outputs a map file as part of the build.
 
 $ bAllowRuntimeSymbolFiles : True if runtime symbols files should be generated as a post build step for some platforms. These files are used by the engine to resolve symbol names of callstack backtraces in logs.
diff --git a/Engine/Documentation/Source/Programming/UnrealBuildSystem/TargetFiles/TargetFilesProperties/TargetFilesProperties.INT.udn b/Engine/Documentation/Source/Programming/UnrealBuildSystem/TargetFiles/TargetFilesProperties/TargetFilesProperties.INT.udn
--- a/Engine/Documentation/Source/Programming/UnrealBuildSystem/TargetFiles/TargetFilesProperties/TargetFilesProperties.INT.udn
+++ b/Engine/Documentation/Source/Programming/UnrealBuildSystem/TargetFiles/TargetFilesProperties/TargetFilesProperties.INT.udn
@@ -430,8 +430,6 @@ $ bForceIncludePCHHeadersForGenCppFilesWhenPCHIsDisabled (Boolean): Whether PCH
 
 $ bPreprocessOnly (Boolean): Whether to just preprocess source files for this target, and skip compilation
 
-$ bPreprocessDepends (Boolean): Generate dependency files by preprocessing. This is only recommended when distributing builds as it adds additional overhead.
-
 $ bWithAssembly (Boolean): Whether to generate assembly data while compiling this target. Works exclusively on MSVC for now.
 
 $ bAllowUbaCompression (Boolean): Experimental: Store object (.obj) compressed on disk. Requires UBA to link, currently MSVC only. Toggling this flag will invalidate MSVC actions. Warning, this option is not currently compatitable with PGO or the the cl-clang linker as those are not detoured and linking will fail.
@@ -492,8 +490,6 @@ $ bCheckLicenseViolations (Boolean): Tells the UBT to check if module currently
 
 $ bBreakBuildOnLicenseViolation (Boolean): Tells the UBT to break build if module currently being built is violating EULA.
 
-$ bUseFastPDBLinking (Nullable&lt;Boolean&gt;): Whether to use the :FASTLINK option when building with /DEBUG to create local PDBs on Windows. No longer recommended.
-
 $ bCreateMapFile (Boolean): Outputs a map file as part of the build.
 
 $ bAllowRuntimeSymbolFiles (Boolean): True if runtime symbols files should be generated as a post build step for some platforms. These files are used by the engine to resolve symbol names of callstack backtraces in logs.
@@ -616,8 +612,12 @@ $ bMergeModuleAndGeneratedUnityFiles (Boolean): Whether to merge module and gene
 
 $ bDetailedUnityFiles (Boolean): Whether to add additional information to the unity files, such as '_of_X' in the file name. Not recommended.
 
+$ bUseFastPDBLinking (Nullable&lt;Boolean&gt;): Whether to use the :FASTLINK option when building with /DEBUG to create local PDBs on Windows. No longer recommended.
+
 $ bDisableDebugInfo (Boolean): Whether to globally disable debug info generation; Obsolete, please use TargetRules.DebugInfoMode instead
 
+$ bPreprocessDepends (Boolean): Generate dependency files by preprocessing. Obsolete.
+
 $ bEnableCppCoroutinesForEvaluation (Boolean): Whether to enable support for C++20 coroutines This option is provided to facilitate evaluation of the feature. Expect the name of the option to change. Use of coroutines in with UE is untested and unsupported.
 
 $ bUseMallocProfiler (Boolean): If true, then enable memory profiling in the build (defines USE_MALLOC_PROFILER=1 and forces bOmitFramePointers=false).
