Documentation for build tools

diff --git a/Engine/Documentation/Source/Programming/Development/BuildGraph/ScriptAnatomy/Tasks/PredefinedTasks/BuildGraphScriptPredefinedTasks.INT.udn b/Engine/Documentation/Source/Programming/Development/BuildGraph/ScriptAnatomy/Tasks/PredefinedTasks/BuildGraphScriptPredefinedTasks.INT.udn
--- a/Engine/Documentation/Source/Programming/Development/BuildGraph/ScriptAnatomy/Tasks/PredefinedTasks/BuildGraphScriptPredefinedTasks.INT.udn
+++ b/Engine/Documentation/Source/Programming/Development/BuildGraph/ScriptAnatomy/Tasks/PredefinedTasks/BuildGraphScriptPredefinedTasks.INT.udn
@@ -604,6 +604,7 @@ Retrieves an artifact from Cloud DDC
 | Commit                | String        | Optional | Change number for the artifact                                                                                                       |
 | RequireMatchingCommit | Boolean       | Optional | Requires that the current synced commit is the same as the artifacts commit                                                          |
 | Name                  | String        | Optional | Name of the artifact                                                                                                                 |
+| BuildGroup            | String        | Optional | Optional filter to require the artifact to be in build group                                                                         |
 | Type                  | String        | Required | The artifact type. Determines the permissions and expiration policy for the artifact.                                                |
 | Keys                  | String        | Optional | Keys for the artifact                                                                                                                |
 | OutputDir             | String        | Required | Output directory for                                                                                                                 |
diff --git a/Engine/Documentation/Source/Programming/UnrealBuildSystem/Configuration/BuildConfigProperties/BuildConfigProperties.INT.udn b/Engine/Documentation/Source/Programming/UnrealBuildSystem/Configuration/BuildConfigProperties/BuildConfigProperties.INT.udn
--- a/Engine/Documentation/Source/Programming/UnrealBuildSystem/Configuration/BuildConfigProperties/BuildConfigProperties.INT.udn
+++ b/Engine/Documentation/Source/Programming/UnrealBuildSystem/Configuration/BuildConfigProperties/BuildConfigProperties.INT.udn
@@ -95,6 +95,8 @@ $ bAutoRTFMClosedStaticLinkage : Whether to link closed function declarations st
 
 $ bUseInlining : Enable inlining for all modules.
 
+$ bWithLiveCoding : Whether to enable support for live coding
+
 $ bUseDebugLiveCodingConsole : Whether to enable support for live coding
 
 $ bUseXGEController : Whether the XGE controller worker and modules should be included in the engine build. These are required for distributed shader compilation using the XGE interception interface.
