Documentation for build tools

diff --git a/Engine/Documentation/Source/Programming/UnrealBuildSystem/TargetFiles/TargetFilesProperties/TargetFilesProperties.INT.udn b/Engine/Documentation/Source/Programming/UnrealBuildSystem/TargetFiles/TargetFilesProperties/TargetFilesProperties.INT.udn
--- a/Engine/Documentation/Source/Programming/UnrealBuildSystem/TargetFiles/TargetFilesProperties/TargetFilesProperties.INT.udn
+++ b/Engine/Documentation/Source/Programming/UnrealBuildSystem/TargetFiles/TargetFilesProperties/TargetFilesProperties.INT.udn
@@ -236,6 +236,8 @@ $ bCompileAgainstEngine (Boolean): Enabled for all builds that include the engin
 
 $ bCompileAgainstCoreUObject (Boolean): Enabled for all builds that include the CoreUObject project.  Disabled only when building standalone apps that only link with Core.
 
+$ bEnableConstInitUObject (Boolean): Enable compilation of UClass and related reflection objects as statically initialized global variables.
+
 $ bCompileAgainstApplicationCore (Boolean): Enabled for builds that need to initialize the ApplicationCore module. Command line utilities do not normally need this.
 
 $ bEnableTrace (Boolean): Enabled for builds that want to use against the Trace module for profiling and diagnostics.
