Documentation for build tools

diff --git a/Engine/Documentation/Source/Programming/UnrealBuildSystem/TargetFiles/TargetFilesProperties/TargetFilesProperties.INT.udn b/Engine/Documentation/Source/Programming/UnrealBuildSystem/TargetFiles/TargetFilesProperties/TargetFilesProperties.INT.udn
--- a/Engine/Documentation/Source/Programming/UnrealBuildSystem/TargetFiles/TargetFilesProperties/TargetFilesProperties.INT.udn
+++ b/Engine/Documentation/Source/Programming/UnrealBuildSystem/TargetFiles/TargetFilesProperties/TargetFilesProperties.INT.udn
@@ -50,6 +50,8 @@ $ GlobalDefinitions (List&lt;String&gt;): Macros to define globally across the w
 
 $ ProjectDefinitions (List&lt;String&gt;): Macros to define across all macros in the project.
 
+$ BuildSettingsVariables (Dictionary&lt;String, String&gt;): Variables that will be stored in BuildSettings module. Can be fetched with BuildSettings::GetVariable
+
 $ AdditionalProperties (AdditionalProperties): Environment variables. Should be used instead of Environment.Set/GetEnvironmentVariable in .build.cs
 
 $ ExtraModuleNames (List&lt;String&gt;): List of additional modules to be compiled into the target.
