Documentation for build tools

diff --git a/Engine/Documentation/Source/Programming/UnrealBuildSystem/Configuration/BuildConfigProperties/BuildConfigProperties.INT.udn b/Engine/Documentation/Source/Programming/UnrealBuildSystem/Configuration/BuildConfigProperties/BuildConfigProperties.INT.udn
--- a/Engine/Documentation/Source/Programming/UnrealBuildSystem/Configuration/BuildConfigProperties/BuildConfigProperties.INT.udn
+++ b/Engine/Documentation/Source/Programming/UnrealBuildSystem/Configuration/BuildConfigProperties/BuildConfigProperties.INT.udn
@@ -533,6 +533,8 @@ $ CacheMaxWorkers : Max number of cache download tasks that can execute in paral
 
 $ CacheShuffle : Shuffle the list of cache servers before each read attempt
 
+$ ForceNoCache : Shuffle the list of cache servers before each read attempt
+
 $ ReportCacheMissReason : Report reason a cache miss happened. Useful when searching for determinism/portability issues
 
 $ CacheLinkActions : Control if link actions should be downloaded from cache or not
@@ -983,6 +985,8 @@ $ NoCompileCommands : Do not create compile commands json files with compiler ar
 
 $ UseVSCodeExtension : Create a workspace file for use with VS Code extension that communicates directly with UBT.
 
+$ LinuxDebuggerType : Debugger type to be used with Linux.  Defaults to cppdbg.
+
 $ bIncludeDocumentation : True if we should include documentation in the generated projects.
 
 $ bUsePrecompiled : True if build targets should pass the -useprecompiled argument.
diff --git a/Engine/Documentation/Source/Programming/UnrealBuildSystem/TargetFiles/TargetFilesProperties/TargetFilesProperties.INT.udn b/Engine/Documentation/Source/Programming/UnrealBuildSystem/TargetFiles/TargetFilesProperties/TargetFilesProperties.INT.udn
--- a/Engine/Documentation/Source/Programming/UnrealBuildSystem/TargetFiles/TargetFilesProperties/TargetFilesProperties.INT.udn
+++ b/Engine/Documentation/Source/Programming/UnrealBuildSystem/TargetFiles/TargetFilesProperties/TargetFilesProperties.INT.udn
@@ -568,6 +568,8 @@ $ bStripExports (Boolean): Experimental: Strip unused exports from libraries. On
 
 $ bMergeModules (Boolean): Experimental: Merge modular modules into combined libraries. Sets LinkType to Modular and enables bStripExports
 
+$ bCheckDelegateLifetime (Boolean): Experimental: Check lifetimes for delegate instances. Recommended for merged modules testing
+
 $ bMergeOnlyGameFeatures (Boolean): Experimental: Whether to merge only game feature plugins, ignoring regular plugins.
 
 $ bReportModuleMergingData (Boolean): Experimental: Whether to report module merging data to the game target.
