Documentation for build tools

diff --git a/Engine/Documentation/Source/Programming/UnrealBuildSystem/ModuleFiles/ModuleFilesProperties/ModuleFilesProperties.INT.udn b/Engine/Documentation/Source/Programming/UnrealBuildSystem/ModuleFiles/ModuleFilesProperties/ModuleFilesProperties.INT.udn
--- a/Engine/Documentation/Source/Programming/UnrealBuildSystem/ModuleFiles/ModuleFilesProperties/ModuleFilesProperties.INT.udn
+++ b/Engine/Documentation/Source/Programming/UnrealBuildSystem/ModuleFiles/ModuleFilesProperties/ModuleFilesProperties.INT.udn
@@ -62,7 +62,7 @@ $ PrivateIncludePathModuleNames (List&lt;String&gt;): List of modules name (no p
 
 $ PrivateDependencyModuleNames (List&lt;String&gt;): List of private dependency module names.  These are modules that our private code depends on but nothing in our public include files depend on.
 
-$ CircularlyReferencedDependentModules (List&lt;String&gt;): Only for legacy reason, should not be used in new code. List of module dependencies that should be treated as circular references.  This modules must have already been added to either the public or private dependent module list.
+$ CircularlyReferencedDependentModules (HashSet&lt;String&gt;): Only for legacy reason, should not be used in new code. List of module dependencies that should be treated as circular references.  This modules must have already been added to either the public or private dependent module list.
 
 $ PublicSystemIncludePaths (List&lt;String&gt;): List of system/library include paths - typically used for External (third party) modules.  These are public stable header file directories that are not checked when resolving header dependencies.
 
