Documentation for build tools

diff --git a/Engine/Documentation/Source/Programming/Development/BuildGraph/ScriptAnatomy/Tasks/PredefinedTasks/BuildGraphScriptPredefinedTasks.INT.udn b/Engine/Documentation/Source/Programming/Development/BuildGraph/ScriptAnatomy/Tasks/PredefinedTasks/BuildGraphScriptPredefinedTasks.INT.udn
--- a/Engine/Documentation/Source/Programming/Development/BuildGraph/ScriptAnatomy/Tasks/PredefinedTasks/BuildGraphScriptPredefinedTasks.INT.udn
+++ b/Engine/Documentation/Source/Programming/Development/BuildGraph/ScriptAnatomy/Tasks/PredefinedTasks/BuildGraphScriptPredefinedTasks.INT.udn
@@ -212,6 +212,22 @@ Compiles C# project files, and their dependencies.
 | UseSystemCompiler | Boolean   | Optional | Whether to use the system toolchain rather than the bundled UE SDK                                                                                  |
 | If                | Condition | Optional | Whether to execute this task. It is ignored if this condition evaluates to false.                                                                   |
 
+### DDCReplayLoad
+
+Spawns an external executable to UpdateDerivedData and waits for it to complete.
+
+|                 |           |          |                                                                                   |
+| --------------- | --------- | -------- | --------------------------------------------------------------------------------- |
+| Exe             | String    | Required | Executable to spawn.                                                              |
+| DDCReplayFile   | String    | Required | The DDC replay load file or tag set to use.                                       |
+| Arguments       | String    | Optional | Arguments for the newly created process.                                          |
+| WorkingDir      | String    | Optional | Working directory for spawning the new task                                       |
+| Environment     | String    | Optional | Environment variables to set                                                      |
+| EnvironmentFile | String    | Optional | File to read environment from                                                     |
+| LogOutput       | Boolean   | Optional | Write output to the log                                                           |
+| ErrorLevel      | Integer   | Optional | The minimum exit code, which is treated as an error.                              |
+| If              | Condition | Optional | Whether to execute this task. It is ignored if this condition evaluates to false. |
+
 ### Delete
 
 Delete a set of files.
