Documentation for build tools

diff --git a/Engine/Documentation/Source/Programming/UnrealBuildSystem/Configuration/BuildConfigProperties/BuildConfigProperties.INT.udn b/Engine/Documentation/Source/Programming/UnrealBuildSystem/Configuration/BuildConfigProperties/BuildConfigProperties.INT.udn
--- a/Engine/Documentation/Source/Programming/UnrealBuildSystem/Configuration/BuildConfigProperties/BuildConfigProperties.INT.udn
+++ b/Engine/Documentation/Source/Programming/UnrealBuildSystem/Configuration/BuildConfigProperties/BuildConfigProperties.INT.udn
@@ -101,6 +101,8 @@ $ bWithLiveCoding : Whether to enable support for live coding
 
 $ bUseDebugLiveCodingConsole : Whether to enable support for live coding
 
+$ bUseCompactSetAsDefault : By default TSparseSet is used as the underlying TSet container. Use this flag to change that to use TCompactSet instead.
+
 $ bUseXGEController : Whether the XGE controller worker and modules should be included in the engine build. These are required for distributed shader compilation using the XGE interception interface.
 
 $ bAlwaysUseUnityForGeneratedFiles : Whether to force generated files to be placed in a unity file even when unity files are otherwise disabled. It is not recommended to disable this unless debugging issues with UHT file generation.
diff --git a/Engine/Documentation/Source/Programming/UnrealBuildSystem/TargetFiles/TargetFilesProperties/TargetFilesProperties.INT.udn b/Engine/Documentation/Source/Programming/UnrealBuildSystem/TargetFiles/TargetFilesProperties/TargetFilesProperties.INT.udn
--- a/Engine/Documentation/Source/Programming/UnrealBuildSystem/TargetFiles/TargetFilesProperties/TargetFilesProperties.INT.udn
+++ b/Engine/Documentation/Source/Programming/UnrealBuildSystem/TargetFiles/TargetFilesProperties/TargetFilesProperties.INT.udn
@@ -322,6 +322,8 @@ $ OptimizationLevel (OptimizationMode): Allows to fine tune optimizations level
 
 $ bDefineForceInlineHintToInline (Boolean): By default UE_FORCEINLINE_HINT is defined to FORCEINLINE. If this flag is true it will instead be defined to 'inline'. This flag should only be true for targets where code size is a concern and/or is being compiled with PGO or LTO. Note that if this flag is true when PGOOptimize is enabled, it should also be true when PGOProfile is enabled.
 
+$ bUseCompactSetAsDefault (Boolean): By default TSparseSet is used as the underlying TSet container. Use this flag to change that to use TCompactSet instead.
+
 $ bUseAndroidJni (Boolean): Whether to use Android Java Native Interface.
 
 $ FPSemantics (FPSemanticsMode): Allows setting the FP semantics.
