Added tooltips to asset search box

diff --git a/Engine/Documentation/Source/Shared/ContentBrowser/ContentBrowser.INT.udn b/Engine/Documentation/Source/Shared/ContentBrowser/ContentBrowser.INT.udn
--- a/Engine/Documentation/Source/Shared/ContentBrowser/ContentBrowser.INT.udn
+++ b/Engine/Documentation/Source/Shared/ContentBrowser/ContentBrowser.INT.udn
@@ -8,6 +8,49 @@ Description:
 The Content Browser is the primary area of the Unreal Editor for creating, importing, organizing, viewing, and modifying content assets within Unreal Editor. It also provides the ability to manage content folders and perform other useful operations on assets, such as renaming, moving, copying, and viewing references. The Content Browser can search for and interact with all assets in the game.
 
 [VAR:ToolTipFullLink]
-Engine/Content/Browser
+content-browser-in-unreal-engine
 [/VAR]
 [/EXCERPT:Tab]
+
+[EXCERPT:AdvancedSearchSyntax]
+
+You can search for your content in the Content Browser using advanced search operators. These operators can refine your search, making it easier for you to find what you need. You can also use them to search for key-value pairs from Asset metadata, and access special key values. 
+Filters can be saved and shared amongst team members.
+
+---
+
+#Operators
+- Equal (= or == or :), NotEqual(!= or !:)
+- Less (<), LessOrEqual (<= or <:), Greater (>), GreaterOrEqual (>= or >:), 
+- Or (OR or || or |), And (AND or && or &), Not (NOT or !)
+
+#String Operations
+- Invert (-)
+- Exact (+)
+- StartsWith, Ends With (...)
+- Quotes (' or ")
+
+#Special Keys
+- The Asset name (Name)
+- The Asset path (Path)
+- The Asset class (Class or Type)
+- The Asset's parent class (ParentClass)
+- The names of any collections that contain the Asset (Collection)
+- Any meta-data tag
+
+#Examples
+1.Find all blueprints whose parent class have Pawn in its name
+**Type==Blueprint AND ParentClass==Pawn**
+1.Find any asset that doesn't have "Foo" in its path or whose name is exactly "Bar"
+**Path=-"Foo" OR Name=+"Bar"**
+1.Find all materials that with a translucent BlendMode and whose names start with "Foo"
+**Type==Material AND (BlendMode==Translucent AND Name=Foo...)**
+
+---
+
+Follow the link for a more extensive description of the search syntax.
+
+[VAR:ToolTipFullLink]
+advanced-search-syntax-in-unreal-engine
+[/VAR]
+[/EXCERPT:AdvancedSearchSyntax]
diff --git a/Engine/Source/Developer/ToolWidgets/Private/Filters/SFilterSearchBox.cpp b/Engine/Source/Developer/ToolWidgets/Private/Filters/SFilterSearchBox.cpp
--- a/Engine/Source/Developer/ToolWidgets/Private/Filters/SFilterSearchBox.cpp
+++ b/Engine/Source/Developer/ToolWidgets/Private/Filters/SFilterSearchBox.cpp
@@ -95,10 +95,11 @@ void SFilterSearchBoxImpl::Construct(const FArguments& InArgs)
 		.AutoWidth()
 		.Padding(-4.f, 0.f, 0.f, 0.f) // Negative padding intentional to have the same padding between all buttons
 		[
-			// Button to save the currently occuring search 
+			// Button to save the currently occurring search 
 			SNew(SButton)
 			.AddMetaData<FTagMetaData>(FTagMetaData(TEXT("CreateFilterFromCurrentSearchButton")))
 			.ContentPadding(0.f)
+			.ToolTipText(LOCTEXT("SaveSearchButtonToolTipText", "Save current search as a filter or collection."))
 			.Visibility_Lambda([this]()
 			{
 				// Only visible if there is a search active currently and the OnSaveSearchClicked delegate is bound
@@ -130,7 +131,7 @@ void SFilterSearchBoxImpl::Construct(const FArguments& InArgs)
 			.ButtonStyle(FAppStyle::Get(), "HoverOnlyButton")
 			.ClickMethod(EButtonClickMethod::MouseDown)
 			.OnClicked(InArgs._OnSearchHistoryClicked)
-			.ToolTipText(LOCTEXT("SearchHistoryToolTipText", "Click to show the Search History"))
+			.ToolTipText(LOCTEXT("SearchHistoryToolTipText", "Show recent searches."))
 			.Visibility_Lambda([this]()
 			{
 				return this->GetSearchHistoryVisibility.Execute();
diff --git a/Engine/Source/Editor/EditorWidgets/Private/SAssetSearchBox.cpp b/Engine/Source/Editor/EditorWidgets/Private/SAssetSearchBox.cpp
--- a/Engine/Source/Editor/EditorWidgets/Private/SAssetSearchBox.cpp
+++ b/Engine/Source/Editor/EditorWidgets/Private/SAssetSearchBox.cpp
@@ -7,6 +7,7 @@
 #include "CoreTypes.h"
 #include "Framework/Application/SlateApplication.h"
 #include "Framework/Views/ITypedTableView.h"
+#include "IDocumentation.h"
 #include "Input/Events.h"
 #include "InputCoreTypes.h"
 #include "Internationalization/Internationalization.h"
@@ -24,6 +25,7 @@
 #include "Widgets/Layout/SBorder.h"
 #include "Widgets/Layout/SBox.h"
 #include "Widgets/SBoxPanel.h"
+#include "Widgets/SToolTip.h"
 #include "Widgets/Text/STextBlock.h"
 #include "Widgets/Views/SListView.h"
 #include "Widgets/Views/STableRow.h"
@@ -95,6 +97,7 @@ void SAssetSearchBox::Construct( const FArguments& InArgs )
 				.OnKeyDownHandler(this, &SAssetSearchBox::HandleKeyDown)
 				.ShowSearchHistory(InArgs._ShowSearchHistory)
 				.OnSaveSearchClicked(InArgs._OnSaveSearchClicked)
+				.ToolTip(IDocumentation::Get()->CreateToolTip(NSLOCTEXT("AssetSearchBox", "FilterSearchHint", "Type here to search"), /*OverrideContent = */nullptr, "Shared/ContentBrowser", "AdvancedSearchSyntax"))
 			]
 			.MenuContent
 				(
