Documentation for build tools

diff --git a/Engine/Documentation/Source/Programming/UnrealBuildSystem/TargetFiles/TargetFilesProperties/TargetFilesProperties.INT.udn b/Engine/Documentation/Source/Programming/UnrealBuildSystem/TargetFiles/TargetFilesProperties/TargetFilesProperties.INT.udn
--- a/Engine/Documentation/Source/Programming/UnrealBuildSystem/TargetFiles/TargetFilesProperties/TargetFilesProperties.INT.udn
+++ b/Engine/Documentation/Source/Programming/UnrealBuildSystem/TargetFiles/TargetFilesProperties/TargetFilesProperties.INT.udn
@@ -280,6 +280,8 @@ $ bFNameOutlineNumber (Boolean): Compile with FName storing the number part in t
 
 $ bWithPushModel (Boolean): When enabled, Push Model Networking support will be compiled in. This can help reduce CPU overhead of networking, at the cost of more memory. Always enabled in editor builds.
 
+$ bExcludeParallelIrisSupport (Boolean): When enabled, certain thread safety mechanisms will be removed from Iris code this may result in a small cpu perf gain if running a single-threaded server
+
 $ bCompileWithStatsWithoutEngine (Boolean): Whether to include stats support even without the engine.
 
 $ bCompileWithPluginSupport (Boolean): Whether to include plugin support.
