Documentation for build tools

diff --git a/Engine/Documentation/Source/Programming/UnrealBuildSystem/Configuration/BuildConfigProperties/BuildConfigProperties.INT.udn b/Engine/Documentation/Source/Programming/UnrealBuildSystem/Configuration/BuildConfigProperties/BuildConfigProperties.INT.udn
--- a/Engine/Documentation/Source/Programming/UnrealBuildSystem/Configuration/BuildConfigProperties/BuildConfigProperties.INT.udn
+++ b/Engine/Documentation/Source/Programming/UnrealBuildSystem/Configuration/BuildConfigProperties/BuildConfigProperties.INT.udn
@@ -319,6 +319,8 @@ $ bUseIpv6 : Whether telemetry/horde should use ipv6 or not. If network does not
 
 $ bForceHeaderGeneration : If true, force header regeneration. Intended for the build machine.
 
+$ bEnableUHTInputCache : If true, enable the UHT input cache for both reading and writing
+
 $ bDoNotBuildUHT : If true, do not build UHT, assume it is already built.
 
 $ bFailIfGeneratedCodeChanges : If true, fail if any of the generated header files is out of date.
