Documentation for build tools

diff --git a/Engine/Documentation/Source/Programming/UnrealBuildSystem/ModuleFiles/ModuleFilesProperties/ModuleFilesProperties.INT.udn b/Engine/Documentation/Source/Programming/UnrealBuildSystem/ModuleFiles/ModuleFilesProperties/ModuleFilesProperties.INT.udn
--- a/Engine/Documentation/Source/Programming/UnrealBuildSystem/ModuleFiles/ModuleFilesProperties/ModuleFilesProperties.INT.udn
+++ b/Engine/Documentation/Source/Programming/UnrealBuildSystem/ModuleFiles/ModuleFilesProperties/ModuleFilesProperties.INT.udn
@@ -210,7 +210,7 @@ $ bUsePrecompiled (Boolean): Whether this module should use precompiled data. Al
 
 $ bAllowConfidentialPlatformDefines (Boolean): Whether this module can use PLATFORM_XXXX style defines, where XXXX is a confidential platform name. This is used to ensure engine or other shared code does not reveal confidential information inside an #if PLATFORM_XXXX block. Licensee game code may want to allow for them, however.
 
-$ bDisableAutoRTFMInstrumentation (Boolean): Disables AutoRTFM instrumentation to this module only when AutoRTFMCompiler is enabled
+$ bDisableAutoRTFMInstrumentation (Boolean): Disables AutoRTFM instrumentation in this module. Only matters when the AutoRTFM compiler is being used. Modules with the `SupportedTargetTypes` attribute where the `TargetType`'s are only `Editor` or `Program` (or both) will be by default disabled for AutoRTFM instrumentation. Otherwise AutoRTFM will by default be enabled.
 
 $ bInjectAutoRTFMAttributeDisables (Boolean): If `bDisableAutoRTFMInstrumentation` is true, will make the modules `*_API` define contain `AUTORTFM_DISABLE` to let us enable errors for code that uses transactionally-disabled code in places that are transactionally-enabled.
 
