Documentation for build tools

diff --git a/Engine/Documentation/Source/Programming/UnrealBuildSystem/Configuration/BuildConfigProperties/BuildConfigProperties.INT.udn b/Engine/Documentation/Source/Programming/UnrealBuildSystem/Configuration/BuildConfigProperties/BuildConfigProperties.INT.udn
--- a/Engine/Documentation/Source/Programming/UnrealBuildSystem/Configuration/BuildConfigProperties/BuildConfigProperties.INT.udn
+++ b/Engine/Documentation/Source/Programming/UnrealBuildSystem/Configuration/BuildConfigProperties/BuildConfigProperties.INT.udn
@@ -591,6 +591,8 @@ $ ProcessPriority : The priority to set for spawned processes. Valid Settings: I
 
 $ bStopCompilationAfterErrors : When enabled, will stop compiling targets after a compile error occurs.
 
+$ StopCompilationAfterNumErrors : When set, will stop compiling targets after a set number of errors occurs.
+
 $ bShowCompilationTimes : Whether to show compilation times along with worst offenders or not.
 
 $ bShowPerActionCompilationTimes : Whether to show compilation times for each executed action
