Documentation for build tools

diff --git a/Engine/Documentation/Source/Programming/UnrealBuildSystem/ModuleFiles/ModuleFilesProperties/ModuleFilesProperties.INT.udn b/Engine/Documentation/Source/Programming/UnrealBuildSystem/ModuleFiles/ModuleFilesProperties/ModuleFilesProperties.INT.udn
--- a/Engine/Documentation/Source/Programming/UnrealBuildSystem/ModuleFiles/ModuleFilesProperties/ModuleFilesProperties.INT.udn
+++ b/Engine/Documentation/Source/Programming/UnrealBuildSystem/ModuleFiles/ModuleFilesProperties/ModuleFilesProperties.INT.udn
@@ -212,7 +212,7 @@ $ bAllowConfidentialPlatformDefines (Boolean): Whether this module can use PLATF
 
 $ bDisableAutoRTFMInstrumentation (Boolean): Disables AutoRTFM instrumentation in this module. Only matters when the AutoRTFM compiler is being used. The default value of this depends on whether any `TargetType` can be deduced for the module. It will default to `true` if:
         * The module has `SupportedTargetTypes` specified and `TargetType.Server` is not one of the target types.
-        * The module is part of a plugin that is `"Type"`: `"Editor"`.
+        * The module is part of a plugin that has a "Type" that is one of "Editor", "EditorAndProgram", or "EditorNoCommandlet".
         * The module is part of a plugin where the module has `"TargetAllowList"` specified and `"Server"` is not one of the allowed target types.
         * The module is part of a plugin where the module has `"TargetDenyList"` specified and `"Server"` is one of the denied target types. Otherwise it will default to `false`.
 
