Documentation for build tools

diff --git a/Engine/Documentation/Source/Programming/UnrealBuildSystem/TargetFiles/TargetFilesProperties/TargetFilesProperties.INT.udn b/Engine/Documentation/Source/Programming/UnrealBuildSystem/TargetFiles/TargetFilesProperties/TargetFilesProperties.INT.udn
--- a/Engine/Documentation/Source/Programming/UnrealBuildSystem/TargetFiles/TargetFilesProperties/TargetFilesProperties.INT.udn
+++ b/Engine/Documentation/Source/Programming/UnrealBuildSystem/TargetFiles/TargetFilesProperties/TargetFilesProperties.INT.udn
@@ -106,6 +106,8 @@ $ OptimizeForSizeAndSpeedModules (String[]): List of modules to optimize for siz
 
 $ Name (String): The name of this target
 
+$ SelectedTestModuleNames (List&lt;String&gt;): The selection of test modules to be exclusively included in compilation (no other modules tests will be included).
+
 $ Logger (ILogger): Logger for output relating to this target. Set before the constructor is run from
 
 $ AdditionalData (Object): Generic nullable object so a user can set additional data in a project's TargetRule and access in a project's ModuleRule without needing to add new properties post-release. For example:
