Documentation for build tools

diff --git a/Engine/Documentation/Source/Programming/UnrealBuildSystem/Configuration/BuildConfigProperties/BuildConfigProperties.INT.udn b/Engine/Documentation/Source/Programming/UnrealBuildSystem/Configuration/BuildConfigProperties/BuildConfigProperties.INT.udn
--- a/Engine/Documentation/Source/Programming/UnrealBuildSystem/Configuration/BuildConfigProperties/BuildConfigProperties.INT.udn
+++ b/Engine/Documentation/Source/Programming/UnrealBuildSystem/Configuration/BuildConfigProperties/BuildConfigProperties.INT.udn
@@ -253,10 +253,6 @@ $ bDisableDebugInfo : Whether to globally disable debug info generation; Obsolet
 
 $ bUndefinedIdentifierErrors : Forces the use of undefined identifiers in conditional expressions to be treated as errors.
 
-$ ActionStallReportTime : Number of second of no completed actions to trigger an action stall report. If zero, stall reports will not be enabled.
-
-$ ActionStallTerminateTime : Number of second of no completed actions to trigger to terminate the queue. If zero, force termination not be enabled.
-
 $ bPrioritizeDependencyScanning : When enabled, dependency scanning is prioritized to avoid being slowed down by running local jobs.
 
 $ bStopSNDBSCompilationAfterErrors : When enabled, SN-DBS will stop compiling targets after a compile error occurs.  Recommended, as it saves computing resources for others.
@@ -579,9 +575,13 @@ $ MsvcCRTRedistVersion : Which MSVC CRT Redist version to use
 
 $ CompilerVersion : Which MSVC Compiler version to use
 
-### ParallelExecutor
+### SNDBS
+
+$ bAllowOverVpn : When set to false, SNDBS will not be enabled when running connected to the coordinator over VPN. Configure VPN-assigned subnets via the VpnSubnets parameter.
+
+$ VpnSubnets : List of subnets containing IP addresses assigned by VPN
 
-$ MaxProcessorCount : Maximum processor count for local execution.
+### ParallelExecutor
 
 $ ProcessorCountMultiplier : Processor count multiplier for local execution. Can be below 1 to reserve CPU for other tasks. When using the local executor (not XGE), run a single action on each CPU core. Note that you can set this to a larger value to get slightly faster build times in many cases, but your computer's responsiveness during compiling may be much worse. This value is ignored if the CPU does not support hyper-threading.
 
@@ -601,16 +601,8 @@ $ bLogActionCommandLines : Whether to log command lines for actions being execut
 
 $ bPrintActionTargetNames : Add target names for each action executed
 
-$ bUseActionWeights : Whether to take into account the Action's weight when determining to do more work or not.
-
 $ bShowCPUUtilization : Whether to show CPU utilization after the work is complete.
 
-### SNDBS
-
-$ bAllowOverVpn : When set to false, SNDBS will not be enabled when running connected to the coordinator over VPN. Configure VPN-assigned subnets via the VpnSubnets parameter.
-
-$ VpnSubnets : List of subnets containing IP addresses assigned by VPN
-
 ### Horde
 
 $ Server : Uri of the Horde server
