Documentation for build tools

diff --git a/Engine/Documentation/Source/Programming/UnrealBuildSystem/Configuration/BuildConfigProperties/BuildConfigProperties.INT.udn b/Engine/Documentation/Source/Programming/UnrealBuildSystem/Configuration/BuildConfigProperties/BuildConfigProperties.INT.udn
--- a/Engine/Documentation/Source/Programming/UnrealBuildSystem/Configuration/BuildConfigProperties/BuildConfigProperties.INT.udn
+++ b/Engine/Documentation/Source/Programming/UnrealBuildSystem/Configuration/BuildConfigProperties/BuildConfigProperties.INT.udn
@@ -445,6 +445,8 @@ $ OptimizeForSizeAndSpeed : List of modules to optimize for size and speed. This
 
 $ bCompileChaosVisualDebuggerSupport : Whether to compile in Chaos Visual Debugger (CVD) support features to record the state of the physics simulation
 
+$ bCompileRewindDebuggerSupport : Whether to compile in RewindDebugger (RWD) support features to record traces for debug primitives and logs
+
 $ bCompileAsioSSLSupport : Whether to compile in ssl support in the third party asio module
 
 ### UnrealBuildAccelerator
@@ -549,6 +551,8 @@ $ bDisableHorde : Disable horde all together
 
 $ bCompileChaosVisualDebuggerSupport : Whether to compile in Chaos Visual Debugger (CVD) support features to record the state of the physics simulation
 
+$ bCompileRewindDebuggerSupport : Whether to compile in RewindDebugger (RWD) support features to record traces for debug primitives and logs
+
 $ bCompileAsioSSLSupport : Whether to compile in ssl support in the third party asio module
 
 ### FASTBuild
diff --git a/Engine/Documentation/Source/Programming/UnrealBuildSystem/TargetFiles/TargetFilesProperties/TargetFilesProperties.INT.udn b/Engine/Documentation/Source/Programming/UnrealBuildSystem/TargetFiles/TargetFilesProperties/TargetFilesProperties.INT.udn
--- a/Engine/Documentation/Source/Programming/UnrealBuildSystem/TargetFiles/TargetFilesProperties/TargetFilesProperties.INT.udn
+++ b/Engine/Documentation/Source/Programming/UnrealBuildSystem/TargetFiles/TargetFilesProperties/TargetFilesProperties.INT.udn
@@ -198,6 +198,8 @@ $ bUseChaosMemoryTracking (Boolean): Whether to compile in chaos memory tracking
 
 $ bCompileChaosVisualDebuggerSupport (Boolean): Whether to compile in Chaos Visual Debugger (CVD) support features to record the state of the physics simulation
 
+$ bCompileRewindDebuggerSupport (Boolean): Whether to compile in RewindDebugger (RWD) support features to record traces for debug primitives and logs
+
 $ bCompileAsioSSLSupport (Boolean): Whether to compile in ssl support in the third party asio module
 
 $ bCompileICU (Boolean): Whether to include ICU unicode/i18n support in Core.
