Documentation for build tools

diff --git a/Engine/Documentation/Source/Programming/UnrealBuildSystem/Configuration/BuildConfigProperties/BuildConfigProperties.INT.udn b/Engine/Documentation/Source/Programming/UnrealBuildSystem/Configuration/BuildConfigProperties/BuildConfigProperties.INT.udn
--- a/Engine/Documentation/Source/Programming/UnrealBuildSystem/Configuration/BuildConfigProperties/BuildConfigProperties.INT.udn
+++ b/Engine/Documentation/Source/Programming/UnrealBuildSystem/Configuration/BuildConfigProperties/BuildConfigProperties.INT.udn
@@ -293,6 +293,8 @@ $ bEnableLibFuzzer : Enables LibFuzzer.
 
 $ bEnableMemorySanitizer : Enables memory sanitizer (MSan)
 
+$ bEnableAutoRTFMSanitizer : Enables AutoRTFM sanitizer (AutoRTFMSan)
+
 $ bTuneDebugInfoForLLDB : Turns on tuning of debug info for LLDB
 
 $ bDisableDumpSyms : Whether to globally disable calling dump_syms
