Documentation for build tools

diff --git a/Engine/Documentation/Source/Programming/Development/BuildGraph/ScriptAnatomy/Tasks/PredefinedTasks/BuildGraphScriptPredefinedTasks.INT.udn b/Engine/Documentation/Source/Programming/Development/BuildGraph/ScriptAnatomy/Tasks/PredefinedTasks/BuildGraphScriptPredefinedTasks.INT.udn
--- a/Engine/Documentation/Source/Programming/Development/BuildGraph/ScriptAnatomy/Tasks/PredefinedTasks/BuildGraphScriptPredefinedTasks.INT.udn
+++ b/Engine/Documentation/Source/Programming/Development/BuildGraph/ScriptAnatomy/Tasks/PredefinedTasks/BuildGraphScriptPredefinedTasks.INT.udn
@@ -218,6 +218,7 @@ Spawns an external executable to UpdateDerivedData and waits for it to complete.
 |                 |           |          |                                                                                   |
 | --------------- | --------- | -------- | --------------------------------------------------------------------------------- |
 | Exe             | String    | Required | Executable to spawn.                                                              |
+| Project         | FileSpec  | Optional | The project to run the tool with.                                                 |
 | DDCReplayFile   | String    | Required | The DDC replay load file or tag set to use.                                       |
 | Arguments       | String    | Optional | Arguments for the newly created process.                                          |
 | WorkingDir      | String    | Optional | Working directory for spawning the new task                                       |
diff --git a/Engine/Documentation/Source/Programming/UnrealBuildSystem/Configuration/BuildConfigProperties/BuildConfigProperties.INT.udn b/Engine/Documentation/Source/Programming/UnrealBuildSystem/Configuration/BuildConfigProperties/BuildConfigProperties.INT.udn
--- a/Engine/Documentation/Source/Programming/UnrealBuildSystem/Configuration/BuildConfigProperties/BuildConfigProperties.INT.udn
+++ b/Engine/Documentation/Source/Programming/UnrealBuildSystem/Configuration/BuildConfigProperties/BuildConfigProperties.INT.udn
@@ -101,6 +101,8 @@ $ bWithLiveCoding : Whether to enable support for live coding
 
 $ bUseDebugLiveCodingConsole : Whether to enable support for live coding
 
+$ bUseSparseSetStrictMode : If enabled sparse set will have additional runtime checks against usage that conflicts with compact sets. This can be used to catch
+
 $ bUseCompactSetAsDefault : By default TSparseSet is used as the underlying TSet container. Use this flag to change that to use TCompactSet instead.
 
 $ bUseXGEController : Whether the XGE controller worker and modules should be included in the engine build. These are required for distributed shader compilation using the XGE interception interface.
diff --git a/Engine/Documentation/Source/Programming/UnrealBuildSystem/TargetFiles/TargetFilesProperties/TargetFilesProperties.INT.udn b/Engine/Documentation/Source/Programming/UnrealBuildSystem/TargetFiles/TargetFilesProperties/TargetFilesProperties.INT.udn
--- a/Engine/Documentation/Source/Programming/UnrealBuildSystem/TargetFiles/TargetFilesProperties/TargetFilesProperties.INT.udn
+++ b/Engine/Documentation/Source/Programming/UnrealBuildSystem/TargetFiles/TargetFilesProperties/TargetFilesProperties.INT.udn
@@ -328,6 +328,8 @@ $ OptimizationLevel (OptimizationMode): Allows to fine tune optimizations level
 
 $ bDefineForceInlineHintToInline (Boolean): By default UE_FORCEINLINE_HINT is defined to FORCEINLINE. If this flag is true it will instead be defined to 'inline'. This flag should only be true for targets where code size is a concern and/or is being compiled with PGO or LTO. Note that if this flag is true when PGOOptimize is enabled, it should also be true when PGOProfile is enabled.
 
+$ bUseSparseSetStrictMode (Boolean): If enabled sparse set will have additional runtime checks against usage that conflicts with compact sets. This can be used to catch
+
 $ bUseCompactSetAsDefault (Boolean): By default TSparseSet is used as the underlying TSet container. Use this flag to change that to use TCompactSet instead.
 
 $ bUseAndroidJni (Boolean): Whether to use Android Java Native Interface.
