Documentation for build tools

diff --git a/Engine/Documentation/Source/Programming/UnrealBuildSystem/Configuration/BuildConfigProperties/BuildConfigProperties.INT.udn b/Engine/Documentation/Source/Programming/UnrealBuildSystem/Configuration/BuildConfigProperties/BuildConfigProperties.INT.udn
--- a/Engine/Documentation/Source/Programming/UnrealBuildSystem/Configuration/BuildConfigProperties/BuildConfigProperties.INT.udn
+++ b/Engine/Documentation/Source/Programming/UnrealBuildSystem/Configuration/BuildConfigProperties/BuildConfigProperties.INT.udn
@@ -537,6 +537,8 @@ $ bDetailedLog : Set to true to see more info about what is happening inside uba
 
 $ AllowDetour : Disable all detouring making all actions run outside uba
 
+$ WritePlaceholders : Disable all detouring making all actions run outside uba
+
 $ CacheMaxWorkers : Max number of cache download tasks that can execute in parallel
 
 $ CacheShuffle : Shuffle the list of cache servers before each read attempt
