Documentation for build tools

diff --git a/Engine/Documentation/Source/Programming/UnrealBuildSystem/Configuration/BuildConfigProperties/BuildConfigProperties.INT.udn b/Engine/Documentation/Source/Programming/UnrealBuildSystem/Configuration/BuildConfigProperties/BuildConfigProperties.INT.udn
--- a/Engine/Documentation/Source/Programming/UnrealBuildSystem/Configuration/BuildConfigProperties/BuildConfigProperties.INT.udn
+++ b/Engine/Documentation/Source/Programming/UnrealBuildSystem/Configuration/BuildConfigProperties/BuildConfigProperties.INT.udn
@@ -537,7 +537,9 @@ $ bDetailedLog : Set to true to see more info about what is happening inside uba
 
 $ AllowDetour : Disable all detouring making all actions run outside uba
 
-$ WritePlaceholders : Disable all detouring making all actions run outside uba
+$ WritePlaceholders : Write zero-byte placeholder files for exe/dll/pdb instead of the full files. This can be useful for validation builds or builds just populating cache
+
+$ SharedMemoryTempFile : Makes the shared memory used in UBA to be backed by a Temp+DeleteOnClose file. This trades committed memory charge for the risk of more I/O under RAM pressure
 
 $ CacheMaxWorkers : Max number of cache download tasks that can execute in parallel
 
