Documentation for build tools

diff --git a/Engine/Documentation/Source/Programming/Development/BuildGraph/ScriptAnatomy/Tasks/PredefinedTasks/BuildGraphScriptPredefinedTasks.INT.udn b/Engine/Documentation/Source/Programming/Development/BuildGraph/ScriptAnatomy/Tasks/PredefinedTasks/BuildGraphScriptPredefinedTasks.INT.udn
--- a/Engine/Documentation/Source/Programming/Development/BuildGraph/ScriptAnatomy/Tasks/PredefinedTasks/BuildGraphScriptPredefinedTasks.INT.udn
+++ b/Engine/Documentation/Source/Programming/Development/BuildGraph/ScriptAnatomy/Tasks/PredefinedTasks/BuildGraphScriptPredefinedTasks.INT.udn
@@ -615,28 +615,29 @@ Retrieves an artifact from Horde
 
 Retrieves an artifact from Cloud DDC
 
-|                       |               |          |                                                                                                                                      |
-| --------------------- | ------------- | -------- | ------------------------------------------------------------------------------------------------------------------------------------ |
-| StreamId              | String        | Optional | Stream containing the artifact                                                                                                       |
-| Commit                | String        | Optional | Change number for the artifact                                                                                                       |
-| RequireMatchingCommit | Boolean       | Optional | Requires that the current synced commit is the same as the artifacts commit                                                          |
-| Name                  | String        | Optional | Name of the artifact                                                                                                                 |
-| BuildGroup            | String        | Optional | Optional filter to require the artifact to be in build group                                                                         |
-| Type                  | String        | Required | The artifact type. Determines the permissions and expiration policy for the artifact.                                                |
-| Keys                  | String        | Optional | Keys for the artifact                                                                                                                |
-| OutputDir             | String        | Required | Output directory for                                                                                                                 |
-| Platform              | String        | Required | The platform the artifact is created for                                                                                             |
-| Project               | FileReference | Optional | The path to the uproject this artifact is created for, this is optional if UAT is run with the -project option                       |
-| Host                  | String        | Optional | The platform the artifact is created for                                                                                             |
-| Namespace             | String        | Optional | The platform the artifact is created for                                                                                             |
-| AccessToken           | String        | Optional | The access token to use                                                                                                              |
-| AllowMultipleMatches  | Boolean       | Optional | Set this to use the latest match if multiple artifacts are possible matches                                                          |
-| AllowMultipart        | Boolean       | Optional | Enable to use multipart endpoints if valuable                                                                                        |
-| HttpVersion           | String        | Optional | Set the explicit http version to use. None to use http handshaking.                                                                  |
-| BoostWorkers          | Boolean       | Optional | Increase the number of worker threads used by zen, may cause machine to be less responsive but will generally improve download times |
-| Include               | FileSpec      | Optional | Optional pattern(s) to download from the artifact, if omitted the entire artifact is downloaded                                      |
-| Exclude               | FileSpec      | Optional | Optional pattern(s) to exclude from the copy for example, Engine/NoCopy*.txt)                                                        |
-| If                    | Condition     | Optional | Whether to execute this task. It is ignored if this condition evaluates to false.                                                    |
+|                       |               |          |                                                                                                                                                                                                                                                                   |
+| --------------------- | ------------- | -------- | ----------------------------------------------------------------------------------------------------------------------------------------------------------------------------------------------------------------------------------------------------------------- |
+| StreamId              | String        | Optional | Stream containing the artifact                                                                                                                                                                                                                                    |
+| Commit                | String        | Optional | Change number for the artifact                                                                                                                                                                                                                                    |
+| RequireMatchingCommit | Boolean       | Optional | Requires that the current synced commit is the same as the artifacts commit                                                                                                                                                                                       |
+| Name                  | String        | Optional | Name of the artifact                                                                                                                                                                                                                                              |
+| BuildGroup            | String        | Optional | Optional filter to require the artifact to be in build group                                                                                                                                                                                                      |
+| Type                  | String        | Required | The artifact type. Determines the permissions and expiration policy for the artifact.                                                                                                                                                                             |
+| Keys                  | String        | Optional | Keys for the artifact                                                                                                                                                                                                                                             |
+| OutputDir             | String        | Required | Output directory for                                                                                                                                                                                                                                              |
+| Platform              | String        | Required | The platform the artifact is created for                                                                                                                                                                                                                          |
+| Project               | FileReference | Optional | The path to the uproject this artifact is created for, this is optional if UAT is run with the -project option                                                                                                                                                    |
+| Host                  | String        | Optional | The platform the artifact is created for                                                                                                                                                                                                                          |
+| Namespace             | String        | Optional | The platform the artifact is created for                                                                                                                                                                                                                          |
+| AccessToken           | String        | Optional | The access token to use                                                                                                                                                                                                                                           |
+| AllowMultipleMatches  | Boolean       | Optional | Set this to use the latest match if multiple artifacts are possible matches                                                                                                                                                                                       |
+| AllowMultipart        | Boolean       | Optional | Enable to use multipart endpoints if valuable                                                                                                                                                                                                                     |
+| HttpVersion           | String        | Optional | Set the explicit http version to use. None to use http handshaking.                                                                                                                                                                                               |
+| BoostWorkers          | Boolean       | Optional | Increase the number of worker threads used by zen, may cause machine to be less responsive but will generally improve download times                                                                                                                              |
+| Include               | FileSpec      | Optional | Optional pattern(s) to download from the artifact, if omitted the entire artifact is downloaded                                                                                                                                                                   |
+| Exclude               | FileSpec      | Optional | Optional pattern(s) to exclude from the copy for example, Engine/NoCopy*.txt)                                                                                                                                                                                     |
+| Append                | Boolean       | Optional | Whether to append downloads to an existing directory (true) or to clear the directory before download (false) Set this to true if you want to download an artifact to a directory, then download an additional artifact from another without wiping state between |
+| If                    | Condition     | Optional | Whether to execute this task. It is ignored if this condition evaluates to false.                                                                                                                                                                                 |
 
 ### SanitizeReceipt
 
