Documentation for build tools

diff --git a/Engine/Documentation/Source/Programming/Development/BuildGraph/ScriptAnatomy/Tasks/PredefinedTasks/BuildGraphScriptPredefinedTasks.INT.udn b/Engine/Documentation/Source/Programming/Development/BuildGraph/ScriptAnatomy/Tasks/PredefinedTasks/BuildGraphScriptPredefinedTasks.INT.udn
--- a/Engine/Documentation/Source/Programming/Development/BuildGraph/ScriptAnatomy/Tasks/PredefinedTasks/BuildGraphScriptPredefinedTasks.INT.udn
+++ b/Engine/Documentation/Source/Programming/Development/BuildGraph/ScriptAnatomy/Tasks/PredefinedTasks/BuildGraphScriptPredefinedTasks.INT.udn
@@ -720,6 +720,7 @@ Strips debugging information from a set of files.
 | Files     | FileSpec             | Required | List of file specifications separated by semicolons (for example, Engine/.../*.pdb), or the name of a tag set. |
 | OutputDir | DirectoryReference   | Optional | Output directory for the stripped files. Defaults to the input path, overwriting the input files.              |
 | Tag       | TagList              | Optional | Tag to be applied to build products of this task.                                                              |
+| StripAll  | Boolean              | Optional | Whether all symbols should be stripped. Not relevant to all platforms.                                         |
 | If        | Condition            | Optional | Whether to execute this task. It is ignored if this condition evaluates to false.                              |
 
 ### Submit
