Various landscape tooltip adjustments

#Synced-CL 49356504

diff --git a/Engine/Documentation/Source/Shared/LandscapeEditor/LandscapeEditor.INT.udn b/Engine/Documentation/Source/Shared/LandscapeEditor/LandscapeEditor.INT.udn
--- a/Engine/Documentation/Source/Shared/LandscapeEditor/LandscapeEditor.INT.udn
+++ b/Engine/Documentation/Source/Shared/LandscapeEditor/LandscapeEditor.INT.udn
@@ -31,12 +31,19 @@ Engine/Landscape/Creation
 [/VAR]
 [/EXCERPT:NewLandscape]
 
-[EXCERPT:SelectComponentTool]
-The Selection tool will allow you to select individual Landscape Components one at a time or multiple ones if you holding down the CRTL key while selecting. 
+[EXCERPT:NewLandscape]
+The will create a new Landscape in your level from scratch. However if needed, you can also import a heightmap that was created using external tools to create a Landscape based off of that heightmap data. 
 [VAR:ToolTipFullLink]
-Engine/Landscape/Editing/ManageMode
+Engine/Landscape/Creation
+[/VAR]
+[/EXCERPT:NewLandscape]
+
+[EXCERPT:ImportExportTool]
+The Import/Export tool will allow you to import or export the landscape's heightmaps (Sculpt) and/or weightmaps (Paint and Visibility) to a single or a set of image files (png, raw...)
+[VAR:ToolTipFullLink]
+Engine/Landscape/Creation
 [/VAR]
-[/EXCERPT:SelectComponentTool]
+[/EXCERPT:ImportExportTool]
 
 [EXCERPT:AddComponentTool]
 The Add Component tool will allow you to create new components for the current Landscape. When the Add Landscape Components is active, the cursor shows a green wireframe where a new Landscape component can be added. 
@@ -47,7 +54,7 @@ Engine/Landscape/Editing/ManageMode
 [/EXCERPT:AddComponentTool]
 
 [EXCERPT:DeleteComponentTool]
-The Delete Component tool will allow you to delete components from the current Landscape. To delete a Landscape component, simply CRTL+Click on the Landscape component that you wish to delete. 
+The Delete Component tool will allow you to delete components from the current Landscape. To delete a Landscape component, simply click on the Landscape component that you wish to delete. 
 [VAR:ToolTipFullLink]
 Engine/Landscape/Editing/ManageMode
 [/VAR]
@@ -61,28 +68,35 @@ Engine/Landscape/Editing/ManageMode
 [/EXCERPT:MoveToLevelTool]
 
 [EXCERPT:ResizeLandscape]
-The **Change Component Size** Landscape tool allows you to change the size of the entire Landscape. Keep in mind that this does not work on individual components but on the entire Landscape as a whole. With this tool you can change the size of the Landscape components and adjust how many sections each component has in the Change Component Size section.
+The Resize Landscape tool allows you to change the size of the entire Landscape. Keep in mind that this does not work on individual components but on the entire Landscape as a whole. With this tool you can change the size of the Landscape components and adjust how many sections each component has in the Change Component Size section.
 [VAR:ToolTipFullLink]
 Engine/Landscape/Editing/ManageMode
 [/VAR]
 [/EXCERPT:ResizeLandscape]
 
 [EXCERPT:SplineTool]
-The Landscape Splines tool is a flexible system for creating linear features that need to conform to the Landscape such as paths or roads. Landscape Splines can even push or pull the terrain underneath them so that the Splines will better sit on the Landscape.  
+The Landscape Splines tool is a flexible system for creating linear features that need to conform to the Landscape such as paths or roads. Landscape Splines can even push or pull the terrain underneath them so that the Splines will better sit on the Landscape. Hold down CTRL and use the Left Mouse button to create spline points. Holding down CTRL and selecting another spline while a spline is selected allows to join the 2 splines. Holding down CTRL and using the Left Mouse button on an existing spline segment will split the spline in 2. Holding down ALT and using Left Mouse dragging on a given spline point will also split the spline segment in the direction of the drag operation. 
 [VAR:ToolTipFullLink]
 Engine/Landscape/Editing/Splines
 [/VAR]
 [/EXCERPT:SplineTool]
 
+[EXCERPT:BlueprintBrushTool]
+The Blueprint Tool allows to setup Blueprint-based brushes on the landscape edit layers in order to procedurally affect the landscape's heightmaps and/or weightmaps.
+[VAR:ToolTipFullLink]
+Engine/Landscape/Editing/ManageMode
+[/VAR]
+[/EXCERPT:BlueprintBrushTool]
+
 [EXCERPT:SculptTool]
-The Sculpt Tool allows you to raise or lower the height of the Landscape height map. To raise the height of the Landscape heightmap, use CRTL + Left Mouse Button. To lower the height of the Landscape height map, use CRTL+SHIFT+Left Mouse Button.  
+The Sculpt Tool allows you to raise or lower the height of the Landscape height map. To raise the height of the Landscape heightmap, use Left Mouse Button. To lower the height of the Landscape height map, hold down SHIFT and use the Left Mouse Button.  
 [VAR:ToolTipFullLink]
 Engine/Landscape/Editing/SculptMode
 [/VAR]
 [/EXCERPT:SculptTool]
 
 [EXCERPT:SmoothTool]
-The Smooth tool helps to smooth the heightmap, getting rid of jagged artifacts that can sometimes arise after using the Sculpting or Erosion tools. The Smooth tool also works with Painting mode to help smooth the transitions between different Landscape materials. To use the Smoothing tool, simply hold down the CRTL key and use the Left Button of your Mouse to smooth out the area under the Brush. 
+The Smooth tool helps to smooth the heightmap, getting rid of jagged artifacts that can sometimes arise after using the Sculpting or Erosion tools. The Smooth tool also works with Painting mode to help smooth the transitions between different Landscape materials. To use the Smoothing tool, simply use the Left Button of your Mouse to smooth out the area under the Brush. 
 ![](Smoothing_Before_And_After_00.png)
 Left is before using the Smoothing tool. Right is after using the Smoothing tool.
 [VAR:ToolTipFullLink]
@@ -91,7 +105,7 @@ Engine/Landscape/Editing/SculptMode
 [/EXCERPT:SmoothTool]
 
 [EXCERPT:FlattenTool]
-The Flatten tool flattens the Landscape to the level that is currently under the mouse when you first active the tool. To use the Flatten tool, simply hold down the CRTL key and use the Left Button of your Mouse to Flatten the area under the Brush.
+The Flatten tool flattens the Landscape to the level that is currently under the mouse when you first active the tool. To use the Flatten tool, simply use the Left Button of your Mouse to Flatten the area under the Brush.
 ![](Flatten_Before_And_After_00.png)
 Left is before using the Flatten Tool. Right is after using the Flatten tool.
 [VAR:ToolTipFullLink]
@@ -100,7 +114,7 @@ Engine/Landscape/Editing/SculptMode
 [/EXCERPT:FlattenTool]
 
 [EXCERPT:RampTool]
-The Ramp Tool enables you to make a ramp between two points on your Landscape. You can even specify a fall off for the sides of the Ramp that determines how sharply the sides of the Ramp will falloff. To use the Ramp tool, simply hold down the CRTL key and use the Left Mouse Button to place where you want the first point to be located on the Landscape. To place the 2nd Ramp point, keep the CTRL key held down and simply click, using the Left Mouse Button, anywhere else on the Landscape. Once your Ramp points are set you can adjust their location by selecting either icon and moving them using the Gizmo. Once you are happy with the location of your Ramp, you can add the Ramp to the Landscape by using the Add Ramp Button from the Tool Settings menu. If you want to undo the Ramp, simply press the Reset Button to remove the Ramp.
+The Ramp Tool enables you to make a ramp between two points on your Landscape. You can even specify a fall off for the sides of the Ramp that determines how sharply the sides of the Ramp will falloff. To use the Ramp tool, simply hold down the CTRL key and use the Left Mouse Button to place where you want the first point to be located on the Landscape. To place the 2nd Ramp point, keep the CTRL key held down and simply click, using the Left Mouse Button, anywhere else on the Landscape. Once your Ramp points are set you can adjust their location by selecting either icon and moving them using the Gizmo. Once you are happy with the location of your Ramp, you can add the Ramp to the Landscape by using the Add Ramp Button from the Tool Settings menu. If you want to undo the Ramp, simply press the Reset Button to remove the Ramp.
 ![](Ramp_Before_And_After_00.png)
 Left is before using the Ramp Tool. Right is after using the Ramp tool.
 [VAR:ToolTipFullLink]
@@ -109,7 +123,7 @@ Engine/Landscape/Editing/SculptMode
 [/EXCERPT:RampTool]
 
 [EXCERPT:ErosionTool]
-The Erosion tool uses a thermal erosion simulation to adjust the height of the heightmap. This simulates the transfer of soil from higher elevations to lower elevations. The larger the difference in elevation, the more erosion will occur. This tool also applies a noise effect on top of the erosion, if desired, to provide a more natural random appearance. To use the Erosion tool, simply hold down the CRTL key and use the Left Mouse Button to apply the Erosion in the area under the Brush.
+The Erosion tool uses a thermal erosion simulation to adjust the height of the heightmap. This simulates the transfer of soil from higher elevations to lower elevations. The larger the difference in elevation, the more erosion will occur. This tool also applies a noise effect on top of the erosion, if desired, to provide a more natural random appearance. To use the Erosion tool, simply use the Left Mouse Button to apply the Erosion in the area under the Brush.
 ![](Erosion_Before_And_After_00.png)
 Left is before using the Erosion Tool. Right is after using the Erosion tool.
 [VAR:ToolTipFullLink]
@@ -118,7 +132,7 @@ Engine/Landscape/Editing/SculptMode
 [/EXCERPT:ErosionTool]
 
 [EXCERPT:HydroErosionTool]
-The Hydro Erosion tool uses a hydraulic erosion simulation, i.e. erosion from water, to adjust the height of the heightmap. A noise filter is used to determine where the initial rain is distributed. Then the simulation is calculated to determine water flow from that initial rain as well as dissolving, water transfer, and evaporation. The result of that calculation provides the actual value used to lower the heightmap. To use the Hydro Erosion tool, simply hold down the CRTL key and use the Left Mouse Button to apply the Hydro Erosion to the area under the Brush.
+The Hydro Erosion tool uses a hydraulic erosion simulation, i.e. erosion from water, to adjust the height of the heightmap. A noise filter is used to determine where the initial rain is distributed. Then the simulation is calculated to determine water flow from that initial rain as well as dissolving, water transfer, and evaporation. The result of that calculation provides the actual value used to lower the heightmap. To use the Hydro Erosion tool, simply use the Left Mouse Button to apply the Hydro Erosion to the area under the Brush.
 ![](Hydro_Erosion_Before_And_After_00.png)
 Left is before using the Hydro Erosion Tool. Right is after using the Hydro Erosion tool.
 [VAR:ToolTipFullLink]
@@ -127,7 +141,7 @@ Engine/Landscape/Editing/SculptMode
 [/EXCERPT:HydroErosionTool]
 
 [EXCERPT:NoiseTool]
-The Noise tool applies a noise filter to the heightmap or layer weight. The strength determines the amount of noise to apply to the heightmap or layer weight. To use the Noise tool, simply hold down the CRTL key and use the Left Mouse Button to apply the noise to the area under the Brush.
+The Noise tool applies a noise filter to the heightmap or layer weight. The strength determines the amount of noise to apply to the heightmap or layer weight. To use the Noise tool, simply use the Left Mouse Button to apply the noise to the area under the Brush.
 ![](Noise_Types_00.png)
 Left Additive Noise, Middle Additive and Subtractive Noise, Right Subtractive Noise.
 [VAR:ToolTipFullLink]
@@ -135,15 +149,8 @@ Engine/Landscape/Editing/SculptMode
 [/VAR]
 [/EXCERPT:NoiseTool]
 
-[EXCERPT:RetopologizeTool]
-The Retopologize tool automatically adjusts Landscape vertices with an X/Y offset map to improve vertex density on cliffs, reducing texture stretching. An X/Y offset map makes the Landscape slower to render and paint on with other tools, so only use the Retopologize tool if needed. To use the Retopologize tool, simply hold down the CRTL key and use the Left Mouse Button to Retopologize the area under the Brush.
-[VAR:ToolTipFullLink]
-Engine/Landscape/Editing/SculptMode
-[/VAR]
-[/EXCERPT:RetopologizeTool]
-
 [EXCERPT:VisibilityTool]
-The Visibility tool allows you to create holes in your Landscape. This is good if you would like to create a hole in your Landscape for something like a cave. To use the Visibility tool, simply hold down the CRTL key and use the Left Mouse Button to change the Visibility to the area under the Brush.
+The Visibility tool allows you to create holes in your Landscape. This is good if you would like to create a hole in your Landscape for something like a cave. To use the Visibility tool, simply use the Left Mouse Button to change the Visibility to the area under the Brush. Hold down SHIFT and use the Left Mouse button to remove the holes. 
 ![](Visibility_Hole_00.png)
 Left: Before the Landscape Visibility is toggled. Right: After the Landscape Visibility is toggled.
 [VAR:ToolTipFullLink]
@@ -152,7 +159,7 @@ Engine/Landscape/Editing/SculptMode
 [/EXCERPT:VisibilityTool]
 
 [EXCERPT:RegionSelectTool]
-The Region Selection tool selects regions of the Landscape using the current brush settings and tool strength to be used to fit a Landscape gizmo to a specific area or to act as a mask for copying data to, or pasting data from, a gizmo. You can also use the Region Selection tool to mask off certain areas of the Landscape so that they are not affected when using other tools such as the Sculpt tool or the Smoothing tool. To use the Region Selection tool, simply paint the area that you want to mask by holding down CRTL and using the Left Mouse Button to paint. To remove the masked area, simply hold down CRTL + SHIFT and Left Mouse Button to erase the mask that was just painted. 
+The Region Selection tool selects regions of the Landscape using the current brush settings and tool strength to be used to fit a Landscape gizmo to a specific area or to act as a mask for copying data to, or pasting data from, a gizmo. You can also use the Region Selection tool to mask off certain areas of the Landscape so that they are not affected when using other tools such as the Sculpt tool or the Smoothing tool. To use the Region Selection tool, simply paint the area that you want to mask by using the Left Mouse Button to paint. To remove the masked area, simply hold down SHIFT and use the Left Mouse Button to erase the mask that was just painted. 
 [VAR:ToolTipFullLink]
 Engine/Landscape/Editing/SculptMode
 [/VAR]
@@ -229,7 +236,7 @@ Engine/Landscape/Editing/Brushes
 [/EXCERPT:CircleBrush_Tip]
 
 [EXCERPT:PaintTool]
-Paint mode enables you to modify the appearance of your Landscape by selectively applying different Material Layers to different parts of your Landscape. To apply different materials to your landscape, simply select the material you wish to paint from the Target Layers and paint by hold down the CRTL key and using the Left Mouse Button to apply the selected material to the area under the Brush. Holding CRTL + SHIFT will erase the material that was just painted. Keep in mind that you can only erase if there is another Material Layer on that texel of the Landscape. This means that areas that are 100% painted over cannot be erased and have to painted over. 
+Paint mode enables you to modify the appearance of your Landscape by selectively applying different Material Layers to different parts of your Landscape. To apply different materials to your landscape, simply select the material you wish to paint from the Target Layers and paint using the Left Mouse Button to apply the selected material to the area under the Brush. Holding SHIFT and using the Left Mouse button will erase the Target Layer that is currently selected. On weight-blended Target Layers, holding CTRL and using the Left Mouse button will deactivate Exclusive Paint Mode, allowing you to affect only the selected target layer (rather than all the layers from the blending group).
 [VAR:ToolTipFullLink]
 Engine/Landscape/Editing/PaintMode
 [/VAR]
diff --git a/Engine/Documentation/Source/Shared/Types/ULandscapeEditorObject/ULandscapeEditorObject.INT.udn b/Engine/Documentation/Source/Shared/Types/ULandscapeEditorObject/ULandscapeEditorObject.INT.udn
--- a/Engine/Documentation/Source/Shared/Types/ULandscapeEditorObject/ULandscapeEditorObject.INT.udn
+++ b/Engine/Documentation/Source/Shared/Types/ULandscapeEditorObject/ULandscapeEditorObject.INT.udn
@@ -18,7 +18,7 @@ Engine/Landscape/Creation
 [/EXCERPT:NewLandscape_QuadsPerSection]
 
 [EXCERPT:NewLandscape_SectionsPerComponent]
-Sections Per Component helps with the Landscapes LOD.  Each section is the basic unit for the Landscape LOD. One component could have 2 x 2 sections which mean�s that is possible that one component could be rendering four different LOD�s at once. With a larger section size you get the added benefit of reduced CPU calculation time. However you could run into issues with the Landscape rendering to many vertices at once. This is could be especially prevalent when using very large areas of Landscape.  These problems could be even worse on mobile devices because of the limited amount of draw calls you can have due to hardware limitations.
+Sections Per Component helps with the Landscapes LOD.  Each section is the basic unit for the Landscape LOD. One component could have 2 x 2 sections which means that is possible that one component could be rendering four different LODs at once. With a larger section size you get the added benefit of reduced CPU calculation time. However you could run into issues with the Landscape rendering to many vertices at once. This is could be especially prevalent when using very large areas of Landscape.  These problems could be even worse on mobile devices because of the limited amount of draw calls you can have due to hardware limitations.
 NewLandscape_SectionsPerComponent    
 [VAR:ToolTipFullLink]
 Engine/Landscape/Creation
diff --git a/Engine/Source/Editor/LandscapeEditor/Private/LandscapeEditorCommands.cpp b/Engine/Source/Editor/LandscapeEditor/Private/LandscapeEditorCommands.cpp
--- a/Engine/Source/Editor/LandscapeEditor/Private/LandscapeEditorCommands.cpp
+++ b/Engine/Source/Editor/LandscapeEditor/Private/LandscapeEditorCommands.cpp
@@ -76,7 +76,7 @@ void FLandscapeEditorCommands::RegisterCommands()
 	UI_COMMAND(VisibilityTool, "Visibility", "Mask out individual quads in the landscape, leaving a hole.", EUserInterfaceActionType::RadioButton, FInputChord());
 	NameToCommandMap.Add("Tool_Visibility", VisibilityTool);
 
-	UI_COMMAND(BlueprintBrushTool, "Blueprint", "Custom painting tools created using Blueprint.", EUserInterfaceActionType::RadioButton, FInputChord());
+	UI_COMMAND(BlueprintBrushTool, "Blueprint", "Custom Blueprint-based brushes that will procedurally affect the landscape's heightmaps and/or weightmaps.", EUserInterfaceActionType::RadioButton, FInputChord());
 	NameToCommandMap.Add("Tool_BlueprintBrush", BlueprintBrushTool);
 
 	UI_COMMAND(SelectComponentTool, "Select", "Select components to use with other tools", EUserInterfaceActionType::RadioButton, FInputChord());
@@ -98,7 +98,7 @@ void FLandscapeEditorCommands::RegisterCommands()
 	UI_COMMAND(MirrorTool, "Mirror", "Copies one side of a landscape to the other, to easily create a mirrored landscape.", EUserInterfaceActionType::RadioButton, FInputChord());
 	NameToCommandMap.Add("Tool_Mirror", MirrorTool);
 
-	UI_COMMAND(SplineTool, "Splines", "Ctrl+click to add control points\n\nHaving a control point selected when you ctrl+click will connect to the new control point with a segment\n\nSpline mesh settings can be found on the details panel when you have segments selected", EUserInterfaceActionType::RadioButton, FInputChord());
+	UI_COMMAND(SplineTool, "Splines", "Create splines that will procedurally affect the landscape's heightmaps and/or weightmaps.", EUserInterfaceActionType::RadioButton, FInputChord());
 	NameToCommandMap.Add("Tool_Splines", SplineTool);
 
 	UI_COMMAND(CircleBrush, "Circle", "Simple circular brush", EUserInterfaceActionType::RadioButton, FInputChord());
diff --git a/Engine/Source/Editor/LandscapeEditor/Private/LandscapeEditorDetailCustomization_ImportExport.cpp b/Engine/Source/Editor/LandscapeEditor/Private/LandscapeEditorDetailCustomization_ImportExport.cpp
--- a/Engine/Source/Editor/LandscapeEditor/Private/LandscapeEditorDetailCustomization_ImportExport.cpp
+++ b/Engine/Source/Editor/LandscapeEditor/Private/LandscapeEditorDetailCustomization_ImportExport.cpp
@@ -293,6 +293,7 @@ void FLandscapeEditorDetailCustomization_ImportExport::CustomizeDetails(IDetailL
 				SNew(STextBlock)
 				.Font(DetailBuilder.GetDetailFont())
 				.Text(LOCTEXT("ImportResolution", "Import Resolution"))
+				.ToolTipText(LOCTEXT("ImportResolutionTooltip", "Possible resolutions for the selected image"))
 			]
 		]
 		.ValueContent()
diff --git a/Engine/Source/Editor/LandscapeEditor/Private/LandscapeEditorDetailCustomization_MiscTools.cpp b/Engine/Source/Editor/LandscapeEditor/Private/LandscapeEditorDetailCustomization_MiscTools.cpp
--- a/Engine/Source/Editor/LandscapeEditor/Private/LandscapeEditorDetailCustomization_MiscTools.cpp
+++ b/Engine/Source/Editor/LandscapeEditor/Private/LandscapeEditorDetailCustomization_MiscTools.cpp
@@ -176,7 +176,9 @@ void FLandscapeEditorDetailCustomization_MiscTools::CustomizeDetails(IDetailLayo
 				.IsChecked(this, &FLandscapeEditorDetailCustomization_MiscTools::GetbUseAutoRotateControlPoint)
 				.Content()
 				[
-					SNew(STextBlock).Text(LOCTEXT("Spline.bUseAutoRotateControlPoint.Selected", "Use Auto Rotate Point"))
+					SNew(STextBlock)
+					.Text(LOCTEXT("Spline.bUseAutoRotateControlPoint.Selected", "Use Auto Rotate Point"))
+					.ToolTipText(LOCTEXT("Spline.bUseAutoRotateControlPoint.Selected.Tooltip", "If enabled, the spline point auto-rotates to maintain a smooth spline after releasing the mouse button."))
 				]
 			]
 		];
@@ -191,7 +193,9 @@ void FLandscapeEditorDetailCustomization_MiscTools::CustomizeDetails(IDetailLayo
 				.IsChecked(this, &FLandscapeEditorDetailCustomization_MiscTools::GetbAlwaysRotateForward)
 				.Content()
 				[
-					SNew(STextBlock).Text(LOCTEXT("Spline.bAlwaysForward.Selected", "Auto-Rotate Always Forward"))
+					SNew(STextBlock)
+					.Text(LOCTEXT("Spline.bAlwaysForward.Selected", "Auto-Rotate Always Forward"))
+					.ToolTipText(LOCTEXT("Spline.bAlwaysForward.Selected.Tooltip", "If enabled, the spline point auto-rotates in a forward direction to maintain a consistent orientation after releasing the mouse button."))
 				]
 			]
 		];
diff --git a/Engine/Source/Editor/LandscapeEditor/Private/LandscapeEditorDetailCustomization_NewLandscape.cpp b/Engine/Source/Editor/LandscapeEditor/Private/LandscapeEditorDetailCustomization_NewLandscape.cpp
--- a/Engine/Source/Editor/LandscapeEditor/Private/LandscapeEditorDetailCustomization_NewLandscape.cpp
+++ b/Engine/Source/Editor/LandscapeEditor/Private/LandscapeEditorDetailCustomization_NewLandscape.cpp
@@ -344,6 +344,7 @@ void FLandscapeEditorDetailCustomization_NewLandscape::CustomizeDetails(IDetailL
 			SNew(STextBlock)
 			.Font(DetailBuilder.GetDetailFont())
 			.Text(LOCTEXT("HeightmapResolution", "Heightmap Resolution"))
+			.ToolTipText(LOCTEXT("ImportResolutionTooltip", "Possible resolutions for the selected image"))
 		]
 	]
 	.ValueContent()
@@ -729,6 +730,7 @@ void FLandscapeEditorDetailCustomization_NewLandscape::CustomizeDetails(IDetailL
 			SNew(SButton)
 			.Visibility_Static(&GetVisibilityOnlyInNewLandscapeMode, ENewLandscapePreviewMode::NewLandscape)
 			.Text(LOCTEXT("FillWorld", "Fill World"))
+			.ToolTipText(LOCTEXT("FillWorldTooltip", "Automatically set the component count such that the landscape fills the entire world (up to the maximum supported landscape resolution)"))
 			.AddMetaData<FTutorialMetaData>(FTutorialMetaData(TEXT("FillWorldButton"), TEXT("LevelEditorToolBox")))
 			.OnClicked(this, &FLandscapeEditorDetailCustomization_NewLandscape::OnFillWorldButtonClicked)
 		]
@@ -739,6 +741,7 @@ void FLandscapeEditorDetailCustomization_NewLandscape::CustomizeDetails(IDetailL
 			SNew(SButton)
 			.Visibility_Static(&GetVisibilityOnlyInNewLandscapeMode, ENewLandscapePreviewMode::ImportLandscape)
 			.Text(LOCTEXT("FitToData", "Fit To Data"))
+			.ToolTipText(LOCTEXT("FitToDataTooltip", "Try to find a section size and number of sections that exactly matches the dimensions of the import data"))
 			.AddMetaData<FTagMetaData>(TEXT("ImportButton"))
 			.OnClicked(this, &FLandscapeEditorDetailCustomization_NewLandscape::OnFitImportDataButtonClicked)
 			.IsEnabled(this, &FLandscapeEditorDetailCustomization_NewLandscape::GetImportButtonIsEnabled)
@@ -750,6 +753,7 @@ void FLandscapeEditorDetailCustomization_NewLandscape::CustomizeDetails(IDetailL
 			SNew(SPrimaryButton)
 			.Visibility_Static(&GetVisibilityOnlyInNewLandscapeMode, ENewLandscapePreviewMode::NewLandscape)
 			.Text(LOCTEXT("Create", "Create"))
+			.ToolTipText(LOCTEXT("CreateTooltip", "Create the landscape based on the currently selected parameters"))
 			.AddMetaData<FTutorialMetaData>(FTutorialMetaData(TEXT("CreateButton"), TEXT("LevelEditorToolBox")))
 			.OnClicked(this, &FLandscapeEditorDetailCustomization_NewLandscape::OnCreateButtonClicked)
 			.IsEnabled(this, &FLandscapeEditorDetailCustomization_NewLandscape::IsCreateButtonEnabled)
@@ -761,6 +765,7 @@ void FLandscapeEditorDetailCustomization_NewLandscape::CustomizeDetails(IDetailL
 			SNew(SButton)
 			.Visibility_Static(&GetVisibilityOnlyInNewLandscapeMode, ENewLandscapePreviewMode::ImportLandscape)
 			.Text(LOCTEXT("Import", "Import"))
+			.ToolTipText(LOCTEXT("ImportTooltip", "Create the landscape based on the currently selected parameters and applies the import data onto it"))
 			.OnClicked(this, &FLandscapeEditorDetailCustomization_NewLandscape::OnCreateButtonClicked)
 			.IsEnabled(this, &FLandscapeEditorDetailCustomization_NewLandscape::GetImportButtonIsEnabled)
 		]
diff --git a/Engine/Source/Editor/LandscapeEditor/Private/LandscapeEditorDetailCustomization_TargetLayers.cpp b/Engine/Source/Editor/LandscapeEditor/Private/LandscapeEditorDetailCustomization_TargetLayers.cpp
--- a/Engine/Source/Editor/LandscapeEditor/Private/LandscapeEditorDetailCustomization_TargetLayers.cpp
+++ b/Engine/Source/Editor/LandscapeEditor/Private/LandscapeEditorDetailCustomization_TargetLayers.cpp
@@ -85,7 +85,7 @@ void FLandscapeEditorStructCustomization_FTargetLayerAssetPath::CustomizeHeader(
 				.FillWidth(1)
 				.HAlign(HAlign_Fill)
 				[
-					PropertyHandle_TargetLayerDirectoryPath->CreatePropertyNameWidget(LOCTEXT("TargetLayerAssetPath", "Default Layer Asset Path"), LOCTEXT("TargetLayerAssetPath_ToolTip", "Set the default Target Layer asset folder"))
+					PropertyHandle_TargetLayerDirectoryPath->CreatePropertyNameWidget(LOCTEXT("TargetLayerAssetPath", "Default Layer Asset Path"), LOCTEXT("TargetLayerAssetPath_ToolTip", "Asset folder where Landscape Layer Info assets will be searched when auto-populating landscape target layers"))
 				]
 		]
 		.ValueContent()
@@ -1675,12 +1675,14 @@ TSharedPtr<SWidget> FLandscapeEditorCustomNodeBuilder_TargetLayers::GenerateRow(
 							[
 								SNew(STextBlock)
 								.Text(LOCTEXT("SubtractiveBlend", "Subtractive Blend"))
+								.ToolTipText(LOCTEXT("SubtractiveBlendTooltip", "If enabled, this target layer's effect on this edit layer will be inverted, "
+									"allowing a given edit layer to effectively decrease the target layer's weight from edit layers underneath."))
 								.ColorAndOpacity_Static(&FLandscapeEditorCustomNodeBuilder_TargetLayers::GetTargetTextColor, Target)
 							]
 						]
 					]
-						+ SVerticalBox::Slot()
-						.AutoHeight()
+					+ SVerticalBox::Slot()
+					.AutoHeight()
 					[
 						SNew(SHorizontalBox)
 						.Visibility_Static(&FLandscapeEditorCustomNodeBuilder_TargetLayers::GetDebugModeColorChannelVisibility, Target)
diff --git a/Engine/Source/Editor/LandscapeEditor/Public/LandscapeEditorObject.h b/Engine/Source/Editor/LandscapeEditor/Public/LandscapeEditorObject.h
--- a/Engine/Source/Editor/LandscapeEditor/Public/LandscapeEditorObject.h
+++ b/Engine/Source/Editor/LandscapeEditor/Public/LandscapeEditorObject.h
@@ -508,6 +508,7 @@ class ULandscapeEditorObject : public UObject
 	UPROPERTY(Category="Tool Settings", EditAnywhere, NonTransactional, Transient, meta=(ShowForTools="ImportExport,CopyPaste"))
 	TWeakObjectPtr<ALandscapeGizmoActiveActor> CurrentGizmoActor;
 
+	// Allows to snap the copied/imported data to a grid
 	UPROPERTY(Category="Tool Settings", EditAnywhere, NonTransactional, meta=(DisplayName="Snap Gizmo to Landscape grid", ShowForTools="CopyPaste,ImportExport"))
 	ELandscapeGizmoSnapType SnapMode;
 
@@ -620,6 +621,7 @@ class ULandscapeEditorObject : public UObject
 	UPROPERTY(NonTransactional)
 	FIntPoint ImportLandscape_GizmoLocalPosition;
 
+	// Describes the transform that should be applied to the import data 
 	UPROPERTY(Category = "Import / Export", EditAnywhere, NonTransactional, meta =(ShowForTools = "ImportExport"))
 	ELandscapeImportTransformType ImportType;
 		
diff --git a/Engine/Source/Editor/LandscapeEditor/Public/LandscapeImportHelper.h b/Engine/Source/Editor/LandscapeEditor/Public/LandscapeImportHelper.h
--- a/Engine/Source/Editor/LandscapeEditor/Public/LandscapeImportHelper.h
+++ b/Engine/Source/Editor/LandscapeEditor/Public/LandscapeImportHelper.h
@@ -49,11 +49,11 @@ UENUM()
 enum class ELandscapeImportTransformType : int8
 {
 
-	None UMETA(DisplayName="Original", ToolTip="Will Import the data at the gizmo location in the original size"),
-	ExpandOffset UMETA(DisplayName="Expand", ToolTip="Will Import the data at the gizmo location and expand the data to fill the landscape") ,
+	None UMETA(DisplayName="Original", ToolTip="Import the data at the gizmo location in the original size"),
+	ExpandOffset UMETA(DisplayName="Expand", ToolTip="Import the data at the gizmo location and expand the data to fill the landscape") ,
 	ExpandCentered UMETA(Hidden), 
-	Resample UMETA(ToolTip="Will resample Import data to fit landscape"),
-	Subregion UMETA(ToolTip="Import Sub-region of the Image to Landscape")
+	Resample UMETA(ToolTip="Resample import data to fit landscape"),
+	Subregion UMETA(ToolTip="Import sub-region of the image to landscape")
 };
 
 FORCEINLINE bool operator==(const FLandscapeImportResolution& Lhs, const FLandscapeImportResolution& Rhs)
diff --git a/Engine/Source/Runtime/Landscape/Classes/LandscapeGizmoActiveActor.h b/Engine/Source/Runtime/Landscape/Classes/LandscapeGizmoActiveActor.h
--- a/Engine/Source/Runtime/Landscape/Classes/LandscapeGizmoActiveActor.h
+++ b/Engine/Source/Runtime/Landscape/Classes/LandscapeGizmoActiveActor.h
@@ -27,8 +27,13 @@ enum ELandscapeGizmoType : int
 UENUM()
 enum class ELandscapeGizmoSnapType
 {
+	/** No snapping of the gizmo */
 	None,
+
+	/** Snap the gizmo to the closest landscape component bounds */
 	Component,
+
+	/** Snap the gizmo to the closest landscape texel (vertex) */
 	Texel
 };
 
