Advanced tooltips adjustements for landscape

#rnx
#Synced-CL 49356504

diff --git a/Engine/Documentation/Source/Shared/LandscapeEditor/LandscapeEditor.INT.udn b/Engine/Documentation/Source/Shared/LandscapeEditor/LandscapeEditor.INT.udn
--- a/Engine/Documentation/Source/Shared/LandscapeEditor/LandscapeEditor.INT.udn
+++ b/Engine/Documentation/Source/Shared/LandscapeEditor/LandscapeEditor.INT.udn
@@ -25,18 +25,18 @@ Engine/Landscape/Materials
 [/EXCERPT:PaintMode]
 
 [EXCERPT:NewLandscape]
-The will create a new Landscape in your level from scratch. However if needed, you can also import a heightmap that was created using external tools to create a Landscape based off of that heightmap data. 
+The New Landscape tool will create a new Landscape in your level from scratch. However if needed, you can also import a heightmap that was created using external tools to create a Landscape based off of that heightmap data. 
 [VAR:ToolTipFullLink]
 Engine/Landscape/Creation
 [/VAR]
 [/EXCERPT:NewLandscape]
 
-[EXCERPT:NewLandscape]
-The will create a new Landscape in your level from scratch. However if needed, you can also import a heightmap that was created using external tools to create a Landscape based off of that heightmap data. 
+[EXCERPT:SelectComponentTool]
+The Selection tool will allow you to select individual Landscape Components one at a time in order to edit their properties.
 [VAR:ToolTipFullLink]
-Engine/Landscape/Creation
+Engine/Landscape/Editing/ManageMode
 [/VAR]
-[/EXCERPT:NewLandscape]
+[/EXCERPT:SelectComponentTool]
 
 [EXCERPT:ImportExportTool]
 The Import/Export tool will allow you to import or export the landscape's heightmaps (Sculpt) and/or weightmaps (Paint and Visibility) to a single or a set of image files (png, raw...)
