Documentation for build tools

diff --git a/Engine/Documentation/Source/Programming/UnrealBuildSystem/Configuration/BuildConfigProperties/BuildConfigProperties.INT.udn b/Engine/Documentation/Source/Programming/UnrealBuildSystem/Configuration/BuildConfigProperties/BuildConfigProperties.INT.udn
--- a/Engine/Documentation/Source/Programming/UnrealBuildSystem/Configuration/BuildConfigProperties/BuildConfigProperties.INT.udn
+++ b/Engine/Documentation/Source/Programming/UnrealBuildSystem/Configuration/BuildConfigProperties/BuildConfigProperties.INT.udn
@@ -541,6 +541,8 @@ $ WritePlaceholders : Write zero-byte placeholder files for exe/dll/pdb instead
 
 $ SharedMemoryTempFile : Makes the shared memory used in UBA to be backed by a Temp+DeleteOnClose file. This trades committed memory charge for the risk of more I/O under RAM pressure
 
+$ MaxRacingPercent : With this enabled local machine will start local processes to race remotes when utilization goes under percent number So if local machine can run 20 actions and only run 10, that means that it is 50% utilized. If MaxRacingPercent is set to 70%, local machine will look at what actions remotes are running and start running them too until it reaches 70% utilization
+
 $ CacheMaxWorkers : Max number of cache download tasks that can execute in parallel
 
 $ CacheShuffle : Shuffle the list of cache servers before each read attempt
