Documentation for build tools

diff --git a/Engine/Documentation/Source/Programming/Development/BuildGraph/ScriptAnatomy/Tasks/PredefinedTasks/BuildGraphScriptPredefinedTasks.INT.udn b/Engine/Documentation/Source/Programming/Development/BuildGraph/ScriptAnatomy/Tasks/PredefinedTasks/BuildGraphScriptPredefinedTasks.INT.udn
--- a/Engine/Documentation/Source/Programming/Development/BuildGraph/ScriptAnatomy/Tasks/PredefinedTasks/BuildGraphScriptPredefinedTasks.INT.udn
+++ b/Engine/Documentation/Source/Programming/Development/BuildGraph/ScriptAnatomy/Tasks/PredefinedTasks/BuildGraphScriptPredefinedTasks.INT.udn
@@ -143,6 +143,27 @@ Copies files from one directory to another.
 | ErrorIfNotFound | Boolean   | Optional | Whether or not to throw an error if no files were found to copy                   |
 | If              | Condition | Optional | Whether to execute this task. It is ignored if this condition evaluates to false. |
 
+### CopyDDC
+
+Spawns an external executable to copy derived data listed in a replay file and waits for it to complete.
+
+|                 |           |          |                                                                                                                          |
+| --------------- | --------- | -------- | ------------------------------------------------------------------------------------------------------------------------ |
+| Project         | FileSpec  | Optional | The project to run the tool with.                                                                                        |
+| Replay          | String    | Required | The DDC replay file or tag set to use as a reference for what is to be copied.                                           |
+| Source          | String    | Required | The DDC graph name or config to use as the source for the data to copy                                                   |
+| Targets         | String    | Optional | An array of the the DDC graph names or configs to use as the target for the data to copy, separated with plus (+) symbol |
+| Preview         | Boolean   | Optional | Optional. Whether to preview the operation without doing any writing of data                                             |
+| Overwrite       | Boolean   | Optional | Optional. Whether to allow overwriting of data in the target(s)                                                          |
+| ForceOverwrite  | Boolean   | Optional | Optional. Whether to force overwriting of data in the target(s) even when it is there and seems identical                |
+| Arguments       | String    | Optional | Arguments for the newly created process.                                                                                 |
+| WorkingDir      | String    | Optional | Working directory for spawning the copy task                                                                             |
+| Environment     | String    | Optional | Environment variables to set                                                                                             |
+| EnvironmentFile | String    | Optional | File to read environment from                                                                                            |
+| LogOutput       | Boolean   | Optional | Write output to the log                                                                                                  |
+| ErrorLevel      | Integer   | Optional | The minimum exit code, which is treated as an error.                                                                     |
+| If              | Condition | Optional | Whether to execute this task. It is ignored if this condition evaluates to false.                                        |
+
 ### CreateArtifact
 
 Uploads an artifact to Horde
