Documentation for build tools

diff --git a/Engine/Documentation/Source/Programming/UnrealBuildSystem/TargetFiles/TargetFilesProperties/TargetFilesProperties.INT.udn b/Engine/Documentation/Source/Programming/UnrealBuildSystem/TargetFiles/TargetFilesProperties/TargetFilesProperties.INT.udn
--- a/Engine/Documentation/Source/Programming/UnrealBuildSystem/TargetFiles/TargetFilesProperties/TargetFilesProperties.INT.udn
+++ b/Engine/Documentation/Source/Programming/UnrealBuildSystem/TargetFiles/TargetFilesProperties/TargetFilesProperties.INT.udn
@@ -210,6 +210,8 @@ $ bCompileISPC (Boolean): Whether to compile using ISPC.
 
 $ bCompilePython (Boolean): Whether to compile in python support
 
+$ bMultiViewPerSlateWindowSupport (Boolean): Whether to compile with support for multiple platform-native views per slate window (right now only implemented on macOS)
+
 $ bUseGameplayDebugger (Boolean): Whether to compile with WITH_GAMEPLAY_DEBUGGER enabled with all Engine's default gameplay debugger categories.
 
 $ bUseGameplayDebuggerCore (Boolean): Set to true when bUseGameplayDebugger is false but GameplayDebugger's core parts are required.
