Documentation for build tools

diff --git a/Engine/Documentation/Source/Programming/UnrealBuildSystem/Configuration/BuildConfigProperties/BuildConfigProperties.INT.udn b/Engine/Documentation/Source/Programming/UnrealBuildSystem/Configuration/BuildConfigProperties/BuildConfigProperties.INT.udn
--- a/Engine/Documentation/Source/Programming/UnrealBuildSystem/Configuration/BuildConfigProperties/BuildConfigProperties.INT.udn
+++ b/Engine/Documentation/Source/Programming/UnrealBuildSystem/Configuration/BuildConfigProperties/BuildConfigProperties.INT.udn
@@ -95,6 +95,8 @@ $ bAutoRTFMVerify : Whether to run LLVM verification after the AutoRTFM compiler
 
 $ bAutoRTFMClosedStaticLinkage : Whether to link closed function declarations statically.
 
+$ bAutoRTFMRedirectionInstrumentation : Whether to instrument all memory accesses for redirection.
+
 $ bUseInlining : Enable inlining for all modules.
 
 $ bWithLiveCoding : Whether to enable support for live coding
diff --git a/Engine/Documentation/Source/Programming/UnrealBuildSystem/TargetFiles/TargetFilesProperties/TargetFilesProperties.INT.udn b/Engine/Documentation/Source/Programming/UnrealBuildSystem/TargetFiles/TargetFilesProperties/TargetFilesProperties.INT.udn
--- a/Engine/Documentation/Source/Programming/UnrealBuildSystem/TargetFiles/TargetFilesProperties/TargetFilesProperties.INT.udn
+++ b/Engine/Documentation/Source/Programming/UnrealBuildSystem/TargetFiles/TargetFilesProperties/TargetFilesProperties.INT.udn
@@ -192,6 +192,8 @@ $ bAutoRTFMVerify (Boolean): Whether to run LLVM verification after the AutoRTFM
 
 $ bAutoRTFMClosedStaticLinkage (Boolean): Whether to link closed function declarations statically.
 
+$ bAutoRTFMRedirectionInstrumentation (Boolean): Whether to instrument all memory accesses for redirection.
+
 $ bUseChaosChecked (Boolean): Whether to compile in checked chaos features for debugging
 
 $ bUseChaosMemoryTracking (Boolean): Whether to compile in chaos memory tracking features
