Documentation for build tools

diff --git a/Engine/Documentation/Source/Programming/UnrealBuildSystem/ModuleFiles/ModuleFilesProperties/ModuleFilesProperties.INT.udn b/Engine/Documentation/Source/Programming/UnrealBuildSystem/ModuleFiles/ModuleFilesProperties/ModuleFilesProperties.INT.udn
--- a/Engine/Documentation/Source/Programming/UnrealBuildSystem/ModuleFiles/ModuleFilesProperties/ModuleFilesProperties.INT.udn
+++ b/Engine/Documentation/Source/Programming/UnrealBuildSystem/ModuleFiles/ModuleFilesProperties/ModuleFilesProperties.INT.udn
@@ -26,10 +26,6 @@ $ StaticAnalyzerAdditionalCheckers (HashSet&lt;String&gt;): The static analyzer
 
 $ StaticAnalyzerPVSDisabledErrors (HashSet&lt;String&gt;): The PVS Studio analysis errors that should be disabled.
 
-$ AutoSdkDirectory (String): The AutoSDK directory for the active host platform
-
-$ EngineDirectory (String): The current engine directory
-
 $ PluginDirectory (String): Property for the directory containing this plugin. Useful for adding paths to third party dependencies.
 
 $ ModuleDirectory (String): Property for the directory containing this module. Useful for adding paths to third party dependencies.
diff --git a/Engine/Documentation/Source/Programming/UnrealBuildSystem/TargetFiles/TargetFilesProperties/TargetFilesProperties.INT.udn b/Engine/Documentation/Source/Programming/UnrealBuildSystem/TargetFiles/TargetFilesProperties/TargetFilesProperties.INT.udn
--- a/Engine/Documentation/Source/Programming/UnrealBuildSystem/TargetFiles/TargetFilesProperties/TargetFilesProperties.INT.udn
+++ b/Engine/Documentation/Source/Programming/UnrealBuildSystem/TargetFiles/TargetFilesProperties/TargetFilesProperties.INT.udn
@@ -256,6 +256,8 @@ $ bEnableTrace (Boolean): Enabled for builds that want to use against the Trace
 
 $ bForceEnableTrace (Boolean): Force enable tracing - used to allow testing programs to verify tracing works as intended.
 
+$ bEnableTraceSecureConnection (Boolean): Enables secure socket connection support for Trace.
+
 $ bCompileAgainstEditor (Boolean): Enabled for editor builds (TargetType.Editor). Can be overridden for programs (TargetType.Program) that would need to compile against editor code. Not available for other target types. Mainly drives the value of WITH_EDITOR.
 
 $ bCompileRecast (Boolean): Whether to compile Recast navmesh generation.
