Documentation for build tools

diff --git a/Engine/Documentation/Source/Programming/UnrealBuildSystem/Configuration/BuildConfigProperties/BuildConfigProperties.INT.udn b/Engine/Documentation/Source/Programming/UnrealBuildSystem/Configuration/BuildConfigProperties/BuildConfigProperties.INT.udn
--- a/Engine/Documentation/Source/Programming/UnrealBuildSystem/Configuration/BuildConfigProperties/BuildConfigProperties.INT.udn
+++ b/Engine/Documentation/Source/Programming/UnrealBuildSystem/Configuration/BuildConfigProperties/BuildConfigProperties.INT.udn
@@ -291,6 +291,8 @@ $ bEnableThreadSanitizer : Enables thread sanitizer (TSan).
 
 $ bEnableUndefinedBehaviorSanitizer : Enables undefined behavior sanitizer (UBSan)
 
+$ bUsePortableToolchain : Experimental: Enable portable toolchain for distributed compilation. Separates SDK-dependent files (local) from portable files (remote). Requires Auto SDK (UE_SDKS_ROOT).
+
 $ bUseDSYMFiles : Enables the generation of .dsym files. This can be disabled to enable faster iteration times during development.
 
 $ bEnableLibFuzzer : Enables LibFuzzer.
@@ -305,6 +307,8 @@ $ bDisableDumpSyms : Whether to globally disable calling dump_syms
 
 $ bDisableStripSymbols : Whether to globally skip symbol stripping
 
+$ bLinuxSampleBasedPGO : If specified along with -PGOProfile, use sample-based PGO instead of instrumented.
+
 $ bEnableInstrumentation : Enables instrumentation.
 
 $ bWriteSolutionOptionFile : Whether to write a solution option (suo) file for the sln.
diff --git a/Engine/Documentation/Source/Programming/UnrealBuildSystem/ModuleFiles/ModuleFilesProperties/ModuleFilesProperties.INT.udn b/Engine/Documentation/Source/Programming/UnrealBuildSystem/ModuleFiles/ModuleFilesProperties/ModuleFilesProperties.INT.udn
--- a/Engine/Documentation/Source/Programming/UnrealBuildSystem/ModuleFiles/ModuleFilesProperties/ModuleFilesProperties.INT.udn
+++ b/Engine/Documentation/Source/Programming/UnrealBuildSystem/ModuleFiles/ModuleFilesProperties/ModuleFilesProperties.INT.udn
@@ -206,6 +206,8 @@ $ bUsePrecompiled (Boolean): Whether this module should use precompiled data. Al
 
 $ bAllowConfidentialPlatformDefines (Boolean): Whether this module can use PLATFORM_XXXX style defines, where XXXX is a confidential platform name. This is used to ensure engine or other shared code does not reveal confidential information inside an #if PLATFORM_XXXX block. Licensee game code may want to allow for them, however.
 
+$ bRequiresPlatformSDK (Boolean): Experimental: Whether this module requires platform specific SDK files. Not supported by all platforms yet
+
 $ bDisableAutoRTFMInstrumentation (Boolean): Disables AutoRTFM instrumentation in this module. Only matters when the AutoRTFM compiler is being used. The default value of this depends on whether any `TargetType` can be deduced for the module. It will default to `true` if:
         * The module has `SupportedTargetTypes` specified and `TargetType.Server` is not one of the target types.
         * The module is part of a plugin that has a "Type" that is one of "Editor", "EditorAndProgram", "EditorNoCommandlet", or "UncookedOnly".
