Documentation for build tools

diff --git a/Engine/Documentation/Source/Programming/UnrealBuildSystem/Configuration/BuildConfigProperties/BuildConfigProperties.INT.udn b/Engine/Documentation/Source/Programming/UnrealBuildSystem/Configuration/BuildConfigProperties/BuildConfigProperties.INT.udn
--- a/Engine/Documentation/Source/Programming/UnrealBuildSystem/Configuration/BuildConfigProperties/BuildConfigProperties.INT.udn
+++ b/Engine/Documentation/Source/Programming/UnrealBuildSystem/Configuration/BuildConfigProperties/BuildConfigProperties.INT.udn
@@ -109,6 +109,8 @@ $ bUseCompactSetAsDefault : By default TSparseSet is used as the underlying TSet
 
 $ bUseXGEController : Whether the XGE controller worker and modules should be included in the engine build. These are required for distributed shader compilation using the XGE interception interface.
 
+$ bIncludeHeaders : Adds header files in included modules to the build.
+
 $ bAlwaysUseUnityForGeneratedFiles : Whether to force generated files to be placed in a unity file even when unity files are otherwise disabled. It is not recommended to disable this unless debugging issues with UHT file generation.
 
 $ bUseAdaptiveUnityBuild : Use a heuristic to determine which files are currently being iterated on and exclude them from unity blobs. This results in faster incremental compile times. For Perforce repositories, the current implementation uses the read-only flag to distinguish the working set, assuming that files will be made writable by the source control system if they are being modified. For Git repositories, the implementation uses the git status command. Source code archives downloaded from Git as .zip files are not supported.
diff --git a/Engine/Documentation/Source/Programming/UnrealBuildSystem/TargetFiles/TargetFilesProperties/TargetFilesProperties.INT.udn b/Engine/Documentation/Source/Programming/UnrealBuildSystem/TargetFiles/TargetFilesProperties/TargetFilesProperties.INT.udn
--- a/Engine/Documentation/Source/Programming/UnrealBuildSystem/TargetFiles/TargetFilesProperties/TargetFilesProperties.INT.udn
+++ b/Engine/Documentation/Source/Programming/UnrealBuildSystem/TargetFiles/TargetFilesProperties/TargetFilesProperties.INT.udn
@@ -212,6 +212,8 @@ $ bCompileISPC (Boolean): Whether to compile using ISPC.
 
 $ bCompilePython (Boolean): Whether to compile in python support
 
+$ bEnableConfigSystem (Boolean): Whether to enable the config system (GConfig)
+
 $ bMultiViewPerSlateWindowSupport (Boolean): Whether to compile with support for multiple platform-native views per slate window (right now only implemented on macOS)
 
 $ bUseGameplayDebugger (Boolean): Whether to compile with WITH_GAMEPLAY_DEBUGGER enabled with all Engine's default gameplay debugger categories.
