Documentation for build tools

diff --git a/Engine/Documentation/Source/Programming/Development/BuildGraph/ScriptAnatomy/Tasks/PredefinedTasks/BuildGraphScriptPredefinedTasks.INT.udn b/Engine/Documentation/Source/Programming/Development/BuildGraph/ScriptAnatomy/Tasks/PredefinedTasks/BuildGraphScriptPredefinedTasks.INT.udn
--- a/Engine/Documentation/Source/Programming/Development/BuildGraph/ScriptAnatomy/Tasks/PredefinedTasks/BuildGraphScriptPredefinedTasks.INT.udn
+++ b/Engine/Documentation/Source/Programming/Development/BuildGraph/ScriptAnatomy/Tasks/PredefinedTasks/BuildGraphScriptPredefinedTasks.INT.udn
@@ -658,6 +658,7 @@ Retrieves an artifact from Cloud DDC
 | Include               | FileSpec      | Optional | Optional pattern(s) to download from the artifact, if omitted the entire artifact is downloaded                                                                                                                                                                   |
 | Exclude               | FileSpec      | Optional | Optional pattern(s) to exclude from the copy for example, Engine/NoCopy*.txt)                                                                                                                                                                                     |
 | Append                | Boolean       | Optional | Whether to append downloads to an existing directory (true) or to clear the directory before download (false) Set this to true if you want to download an artifact to a directory, then download an additional artifact from another without wiping state between |
+| Clean                 | Boolean       | Optional | Delete all data in target folder that is not part of the downloaded content                                                                                                                                                                                       |
 | If                    | Condition     | Optional | Whether to execute this task. It is ignored if this condition evaluates to false.                                                                                                                                                                                 |
 
 ### SanitizeReceipt
