Documentation for build tools

diff --git a/Engine/Documentation/Source/Programming/UnrealBuildSystem/Configuration/BuildConfigProperties/BuildConfigProperties.INT.udn b/Engine/Documentation/Source/Programming/UnrealBuildSystem/Configuration/BuildConfigProperties/BuildConfigProperties.INT.udn
--- a/Engine/Documentation/Source/Programming/UnrealBuildSystem/Configuration/BuildConfigProperties/BuildConfigProperties.INT.udn
+++ b/Engine/Documentation/Source/Programming/UnrealBuildSystem/Configuration/BuildConfigProperties/BuildConfigProperties.INT.udn
@@ -167,7 +167,7 @@ $ bUseVFS : Experimental work in progress: Set flags to use toolchain virtual fi
 
 $ CacheSalt : Salt value to modify cache bucket calculation, can be used to invalidate cached actions outputs
 
-$ bChainPCHs : Whether PCHs should be chained when compiling with clang.
+$ bChainPCHs : Whether PCHs should be chained when compiling with clang. Note: The earliest PCH a header is included in the chain will be included, so if there is a mismatch in preprocessor state this can cause an unexpected version of a header to be included. Currently disabled as this is the case.
 
 $ bForceIncludePCHHeadersForGenCppFilesWhenPCHIsDisabled : Whether PCH headers should be force included for gen.cpp files when PCH is disabled.
 
diff --git a/Engine/Documentation/Source/Programming/UnrealBuildSystem/TargetFiles/TargetFilesProperties/TargetFilesProperties.INT.udn b/Engine/Documentation/Source/Programming/UnrealBuildSystem/TargetFiles/TargetFilesProperties/TargetFilesProperties.INT.udn
--- a/Engine/Documentation/Source/Programming/UnrealBuildSystem/TargetFiles/TargetFilesProperties/TargetFilesProperties.INT.udn
+++ b/Engine/Documentation/Source/Programming/UnrealBuildSystem/TargetFiles/TargetFilesProperties/TargetFilesProperties.INT.udn
@@ -448,7 +448,7 @@ $ bUseVFS (Boolean): Experimental work in progress: Set flags to use toolchain v
 
 $ CacheSalt (String): Salt value to modify cache bucket calculation, can be used to invalidate cached actions outputs
 
-$ bChainPCHs (Boolean): Whether PCHs should be chained when compiling with clang.
+$ bChainPCHs (Boolean): Whether PCHs should be chained when compiling with clang. Note: The earliest PCH a header is included in the chain will be included, so if there is a mismatch in preprocessor state this can cause an unexpected version of a header to be included. Currently disabled as this is the case.
 
 $ bForceIncludePCHHeadersForGenCppFilesWhenPCHIsDisabled (Boolean): Whether PCH headers should be force included for gen.cpp files when PCH is disabled.
 
