Documentation for build tools

diff --git a/Engine/Documentation/Source/Programming/UnrealBuildSystem/Configuration/BuildConfigProperties/BuildConfigProperties.INT.udn b/Engine/Documentation/Source/Programming/UnrealBuildSystem/Configuration/BuildConfigProperties/BuildConfigProperties.INT.udn
--- a/Engine/Documentation/Source/Programming/UnrealBuildSystem/Configuration/BuildConfigProperties/BuildConfigProperties.INT.udn
+++ b/Engine/Documentation/Source/Programming/UnrealBuildSystem/Configuration/BuildConfigProperties/BuildConfigProperties.INT.udn
@@ -1169,3 +1169,11 @@ $ RepositoryPath : Sets the path to use for the repository. Interpreted relative
 
 $ GitPath : Sets the path to use for the Git executable. Defaults to "git" (assuming it is in the PATH).
 
+### GDKPlatform
+
+$ ContentOnlyDebugProject : Path to the .uproject file to use when F5 debugging a content-only project via the UE5 engine project
+
+$ bVerifyLibGDKEditions : Whether to verify that any libs the application links to are using the same GDK edition
+
+$ bVerifyDLLGDKEditions : Whether to verify that any DLLs the application links to are using the same GDK edition
+
diff --git a/Engine/Documentation/Source/Programming/UnrealBuildSystem/TargetFiles/TargetFilesProperties/TargetFilesProperties.INT.udn b/Engine/Documentation/Source/Programming/UnrealBuildSystem/TargetFiles/TargetFilesProperties/TargetFilesProperties.INT.udn
--- a/Engine/Documentation/Source/Programming/UnrealBuildSystem/TargetFiles/TargetFilesProperties/TargetFilesProperties.INT.udn
+++ b/Engine/Documentation/Source/Programming/UnrealBuildSystem/TargetFiles/TargetFilesProperties/TargetFilesProperties.INT.udn
@@ -456,7 +456,7 @@ $ bPreprocessOnly (Boolean): Whether to just preprocess source files for this ta
 
 $ bWithAssembly (Boolean): Whether to generate assembly data while compiling this target. Works exclusively on MSVC for now.
 
-$ bAllowUbaCompression (Boolean): Experimental: Store object (.obj) compressed on disk. Requires UBA to link, currently MSVC only. Toggling this flag will invalidate MSVC actions. Warning, this option is not currently compatitable with PGO or the the cl-clang linker as those are not detoured and linking will fail.
+$ bAllowUbaCompression (Boolean): Experimental: Store object (.obj) compressed on disk. Requires UBA to link. Toggling this flag will invalidate actions.
 
 $ StaticAnalyzer (StaticAnalyzer): Whether static code analysis should be enabled.
 
