Documentation for build tools

diff --git a/Engine/Documentation/Source/Programming/UnrealBuildSystem/ModuleFiles/ModuleFilesProperties/ModuleFilesProperties.INT.udn b/Engine/Documentation/Source/Programming/UnrealBuildSystem/ModuleFiles/ModuleFilesProperties/ModuleFilesProperties.INT.udn
--- a/Engine/Documentation/Source/Programming/UnrealBuildSystem/ModuleFiles/ModuleFilesProperties/ModuleFilesProperties.INT.udn
+++ b/Engine/Documentation/Source/Programming/UnrealBuildSystem/ModuleFiles/ModuleFilesProperties/ModuleFilesProperties.INT.udn
@@ -238,6 +238,8 @@ $ CStandard (Nullable&lt;CStandardVersion&gt;): Which standard to use for compil
 
 $ ModuleSymbolVisibility (SymbolVisibility): Control visibility of symbols
 
+$ ExternalMetadata (ExternalModuleMetadata): The metadata for external (third-party) modules. Should be non- if and only if the module is of type .
+
 $ bUseAVX (Boolean): Obsolete: Direct the compiler to generate AVX instructions wherever SSE or AVX intrinsics are used, on the platforms that support it. Note that by enabling this you are changing the minspec for the PC platform, and the resultant executable will crash on machines without AVX support.
 
 $ bEnforceIWYU (Boolean): Enforce "include what you use" rules when PCHUsage is set to ExplicitOrSharedPCH; warns when monolithic headers (Engine.h, UnrealEd.h, etc...) are used, and checks that source files include their matching header first.
