Documentation for build tools

diff --git a/Engine/Documentation/Source/Programming/UnrealBuildSystem/Configuration/BuildConfigProperties/BuildConfigProperties.INT.udn b/Engine/Documentation/Source/Programming/UnrealBuildSystem/Configuration/BuildConfigProperties/BuildConfigProperties.INT.udn
--- a/Engine/Documentation/Source/Programming/UnrealBuildSystem/Configuration/BuildConfigProperties/BuildConfigProperties.INT.udn
+++ b/Engine/Documentation/Source/Programming/UnrealBuildSystem/Configuration/BuildConfigProperties/BuildConfigProperties.INT.udn
@@ -385,6 +385,8 @@ $ PCHMemoryAllocationFactor : Determines the amount of memory that the compiler
 
 $ AdditionalLinkerOptions : Allow the target to specify extra options for linking that aren't otherwise noted here
 
+$ bSyntaxCheckOnly : Tells the compiler to check only the syntax of the source files on the command line. (/Zs)
+
 $ bReducedOptimizeHugeFunctions : Whether to reduce optimizations for huge functions over an instruction threshold to improve compile time https://devblogs.microsoft.com/cppblog/msvc-backend-updates-in-visual-studio-2019-versions-16-3-and-16-4/
 
 $ ReducedOptimizeHugeFunctionsThreshold : The instruction threshold to use when reducing optimizations for huge functions, default 20000.
