Documentation for build tools

diff --git a/Engine/Documentation/Source/Programming/UnrealBuildSystem/ModuleFiles/ModuleFilesProperties/ModuleFilesProperties.INT.udn b/Engine/Documentation/Source/Programming/UnrealBuildSystem/ModuleFiles/ModuleFilesProperties/ModuleFilesProperties.INT.udn
--- a/Engine/Documentation/Source/Programming/UnrealBuildSystem/ModuleFiles/ModuleFilesProperties/ModuleFilesProperties.INT.udn
+++ b/Engine/Documentation/Source/Programming/UnrealBuildSystem/ModuleFiles/ModuleFilesProperties/ModuleFilesProperties.INT.udn
@@ -128,9 +128,9 @@ $ BinariesSubFolder (String): Subfolder of Binaries/PLATFORM folder to put this
 
 $ OptimizeCode (CodeOptimization): When this module's code should be optimized.
 
-$ OptimizationLevel (OptimizationMode): Allows fine tuning optimization level for speed and\or code size. This requires a private PCH (or NoPCHs, which is not recommended)
+$ OptimizationLevel (OptimizationMode): Allows fine tuning optimization level for speed and\or code size.
 
-$ FPSemantics (FPSemanticsMode): Allows overriding the FP semantics for this module. This requires a private PCH (or NoPCHs, which is not recommended)
+$ FPSemantics (FPSemanticsMode): Allows overriding the FP semantics for this module.
 
 $ PrivatePCHHeaderFile (String): Explicit private PCH for this module. Implies that this module will not use a shared PCH.
 
