Documentation for build tools

diff --git a/Engine/Documentation/Source/Programming/Development/BuildGraph/ScriptAnatomy/Tasks/PredefinedTasks/BuildGraphScriptPredefinedTasks.INT.udn b/Engine/Documentation/Source/Programming/Development/BuildGraph/ScriptAnatomy/Tasks/PredefinedTasks/BuildGraphScriptPredefinedTasks.INT.udn
--- a/Engine/Documentation/Source/Programming/Development/BuildGraph/ScriptAnatomy/Tasks/PredefinedTasks/BuildGraphScriptPredefinedTasks.INT.udn
+++ b/Engine/Documentation/Source/Programming/Development/BuildGraph/ScriptAnatomy/Tasks/PredefinedTasks/BuildGraphScriptPredefinedTasks.INT.udn
@@ -204,6 +204,7 @@ Uploads an artifact to Cloud DDC
 | StreamId            | String        | Optional | Stream containing the artifact.                                                                                                                                                                                          |
 | Commit              | String        | Optional | Commit for the uploaded artifact.                                                                                                                                                                                        |
 | Files               | FileSpec      | Required | Files to include in the artifact.                                                                                                                                                                                        |
+| SymbolFiles         | FileSpec      | Optional | Files to include in the symbols artifact.                                                                                                                                                                                |
 | Metadata            | String        | Optional | Other metadata for the artifact, separated by semicolons. Use key=value for each field.                                                                                                                                  |
 | Force               | Boolean       | Optional | Force the blobs to be uploaded no matter if they exist or not                                                                                                                                                            |
 | BaselineBranch      | String        | Optional | The streamid used for your mainline, used to build content as patches on top of this branch if no other baselines exist.                                                                                                 |
@@ -835,6 +836,18 @@ Extract files from a zip archive.
 | Tag            | TagList            | Optional | Tag to be applied to the extracted files.                                         |
 | If             | Condition          | Optional | Whether to execute this task. It is ignored if this condition evaluates to false. |
 
+### UploadCloudSymbols
+
+Task that strips symbols from a set of files.
+
+|             |           |          |                                                                                   |
+| ----------- | --------- | -------- | --------------------------------------------------------------------------------- |
+| Files       | FileSpec  | Required | List of output files. PDBs will be extracted from this list.                      |
+| Host        | String    | Optional | The cloud host to which symbols should be uploaded                                |
+| AccessToken | String    | Optional | The access token to use                                                           |
+| Namespace   | String    | Required | The destination namespace to which symbols should be uploaded                     |
+| If          | Condition | Optional | Whether to execute this task. It is ignored if this condition evaluates to false. |
+
 ### Wait
 
 Waits a defined number of seconds.
