Documentation for build tools

diff --git a/Engine/Documentation/Source/Programming/UnrealBuildSystem/Configuration/BuildConfigProperties/BuildConfigProperties.INT.udn b/Engine/Documentation/Source/Programming/UnrealBuildSystem/Configuration/BuildConfigProperties/BuildConfigProperties.INT.udn
--- a/Engine/Documentation/Source/Programming/UnrealBuildSystem/Configuration/BuildConfigProperties/BuildConfigProperties.INT.udn
+++ b/Engine/Documentation/Source/Programming/UnrealBuildSystem/Configuration/BuildConfigProperties/BuildConfigProperties.INT.udn
@@ -117,6 +117,8 @@ $ bUseAdaptiveUnityBuild : Use a heuristic to determine which files are currentl
 
 $ bAdaptiveUnityDisablesOptimizations : Disable optimization for files that are in the adaptive non-unity working set.
 
+$ AdaptiveUnityDisablesOptimizationsConfigurations : Disable optimizations for files that are in the adaptive non-unity working set for the given configurations.
+
 $ bAdaptiveUnityDisablesPCH : Disables force-included PCHs for files that are in the adaptive non-unity working set.
 
 $ bAdaptiveUnityDisablesProjectPCHForProjectPrivate : Backing storage for bAdaptiveUnityDisablesProjectPCH.
diff --git a/Engine/Documentation/Source/Programming/UnrealBuildSystem/TargetFiles/TargetFilesProperties/TargetFilesProperties.INT.udn b/Engine/Documentation/Source/Programming/UnrealBuildSystem/TargetFiles/TargetFilesProperties/TargetFilesProperties.INT.udn
--- a/Engine/Documentation/Source/Programming/UnrealBuildSystem/TargetFiles/TargetFilesProperties/TargetFilesProperties.INT.udn
+++ b/Engine/Documentation/Source/Programming/UnrealBuildSystem/TargetFiles/TargetFilesProperties/TargetFilesProperties.INT.udn
@@ -388,6 +388,8 @@ $ bUseAdaptiveUnityBuild (Boolean): Use a heuristic to determine which files are
 
 $ bAdaptiveUnityDisablesOptimizations (Boolean): Disable optimization for files that are in the adaptive non-unity working set.
 
+$ AdaptiveUnityDisablesOptimizationsConfigurations (UnrealTargetConfiguration[]): Disable optimizations for files that are in the adaptive non-unity working set for the given configurations.
+
 $ bAdaptiveUnityDisablesPCH (Boolean): Disables force-included PCHs for files that are in the adaptive non-unity working set.
 
 $ bAdaptiveUnityDisablesProjectPCHForProjectPrivate (Nullable&lt;Boolean&gt;): Backing storage for bAdaptiveUnityDisablesProjectPCH.
