Documentation for build tools

diff --git a/Engine/Documentation/Source/Programming/Development/BuildGraph/ScriptAnatomy/Tasks/PredefinedTasks/BuildGraphScriptPredefinedTasks.INT.udn b/Engine/Documentation/Source/Programming/Development/BuildGraph/ScriptAnatomy/Tasks/PredefinedTasks/BuildGraphScriptPredefinedTasks.INT.udn
--- a/Engine/Documentation/Source/Programming/Development/BuildGraph/ScriptAnatomy/Tasks/PredefinedTasks/BuildGraphScriptPredefinedTasks.INT.udn
+++ b/Engine/Documentation/Source/Programming/Development/BuildGraph/ScriptAnatomy/Tasks/PredefinedTasks/BuildGraphScriptPredefinedTasks.INT.udn
@@ -977,6 +977,7 @@ Compresses files into a zip archive.
 | Files           | FileSpec           | Optional | List of file specifications separated by semicolons (for example, *.cpp;Engine/.../*.bat), or the name of a tag set. Relative paths are taken from FromDir. |
 | ExecutableFiles | FileSpec           | Optional | List of files that should have an executable bit set.                                                                                                       |
 | ZipFile         | FileReference      | Required | The zip file to create.                                                                                                                                     |
+| SkipDefaultTag  | Boolean            | Optional | Skips adding the resulting zip to the default tagset.                                                                                                       |
 | Tag             | TagList            | Optional | Tag to be applied to the created zip file.                                                                                                                  |
 | If              | Condition          | Optional | Whether to execute this task. It is ignored if this condition evaluates to false.                                                                           |
 
