Documentation for build tools

diff --git a/Engine/Documentation/Source/Programming/UnrealBuildSystem/Configuration/BuildConfigProperties/BuildConfigProperties.INT.udn b/Engine/Documentation/Source/Programming/UnrealBuildSystem/Configuration/BuildConfigProperties/BuildConfigProperties.INT.udn
--- a/Engine/Documentation/Source/Programming/UnrealBuildSystem/Configuration/BuildConfigProperties/BuildConfigProperties.INT.udn
+++ b/Engine/Documentation/Source/Programming/UnrealBuildSystem/Configuration/BuildConfigProperties/BuildConfigProperties.INT.udn
@@ -85,6 +85,8 @@ $ bDebugBuildsActuallyUseDebugCRT : Enables the debug C++ runtime (CRT) for debu
 
 $ bLegalToDistributeBinary : Whether the output from this target can be publicly distributed, even if it has dependencies on modules that are in folders with special restrictions (eg. CarefullyRedist, NotForLicensees, NoRedist).
 
+$ bBuildAllModules : Build all the modules that are valid for this target type. Used for CIS and making installed engine builds.
+
 $ bUseVerseBPVM : Whether to use the BPVM to run Verse.
 
 $ bForceNoAutoRTFMCompiler : Whether to use force AutoRTFM Clang compiler off.
@@ -233,6 +235,8 @@ $ CrashDiagnosticDirectory : Directory where to put crash report files for platf
 
 $ bCheckSystemHeadersForModification : Whether headers in system paths should be checked for modification when determining outdated actions.
 
+$ bDisableLinking : Whether to disable linking for this target.
+
 $ bFlushBuildDirOnRemoteMac : Whether to clean Builds directory on a remote Mac before building.
 
 $ bPrintToolChainTimingInfo : Whether to write detailed timing info from the compiler and linker.
