Documentation for build tools

diff --git a/Engine/Documentation/Source/Programming/UnrealBuildSystem/TargetFiles/TargetFilesProperties/TargetFilesProperties.INT.udn b/Engine/Documentation/Source/Programming/UnrealBuildSystem/TargetFiles/TargetFilesProperties/TargetFilesProperties.INT.udn
--- a/Engine/Documentation/Source/Programming/UnrealBuildSystem/TargetFiles/TargetFilesProperties/TargetFilesProperties.INT.udn
+++ b/Engine/Documentation/Source/Programming/UnrealBuildSystem/TargetFiles/TargetFilesProperties/TargetFilesProperties.INT.udn
@@ -94,6 +94,8 @@ $ DisableMergingModuleAndGeneratedFilesInUnityFiles (String[]): List of modules
 
 $ bShowIncludes (Boolean): Print out files that are included by each source file
 
+$ StaticAnalyzerPVSConfigs (List&lt;String&gt;): Paths to PVS-Studio .pvsconfig configuration files
+
 $ DisableUnityBuildForModules (String[]): List of modules to disable unity builds for
 
 $ EnableOptimizeCodeForModules (String[]): List of modules to enable optimizations for
